Copying //Tasks/UE4/Dev-Online-Plugins-v3 to //UE4/Dev-Main (Source: //Tasks/UE4/Dev-Online-Plugins-v3 @ 3047294)

diff --git a/Engine/Plugins/Online/OnlineFramework/Source/Hotfix/Private/OnlineHotfixManager.cpp b/Engine/Plugins/Online/OnlineFramework/Source/Hotfix/Private/OnlineHotfixManager.cpp
new file mode 100644
--- /dev/null
+++ b/Engine/Plugins/Online/OnlineFramework/Source/Hotfix/Private/OnlineHotfixManager.cpp
@@ -0,0 +1,1028 @@
+// Copyright 1998-2016 Epic Games, Inc. All Rights Reserved.
+
+#include "HotfixPrivatePCH.h"
+#include "OnlineHotfixManager.h"
+#include "Internationalization.h"
+#include "OnlineSubsystemUtils.h"
+
+DEFINE_LOG_CATEGORY(LogHotfixManager);
+
+FName NAME_HotfixManager(TEXT("HotfixManager"));
+
+class FPakFileVisitor : public IPlatformFile::FDirectoryVisitor
+{
+public:
+	virtual bool Visit(const TCHAR* FilenameOrDirectory, bool bIsDirectory) override
+	{
+		if (!bIsDirectory)
+		{
+			Files.Add(FilenameOrDirectory);
+		}
+		return true;
+	}
+
+	TArray<FString> Files;
+};
+
+UOnlineHotfixManager::UOnlineHotfixManager() :
+	Super(),
+	TotalFiles(0),
+	NumDownloaded(0),
+	TotalBytes(0),
+	NumBytes(0),
+	bHotfixingInProgress(false),
+	bHotfixNeedsMapReload(false),
+	ChangedOrRemovedPakCount(0)
+{
+	OnEnumerateFilesCompleteDelegate = FOnEnumerateFilesCompleteDelegate::CreateUObject(this, &UOnlineHotfixManager::OnEnumerateFilesComplete);
+	OnReadFileProgressDelegate = FOnReadFileProgressDelegate::CreateUObject(this, &UOnlineHotfixManager::OnReadFileProgress);
+	OnReadFileCompleteDelegate = FOnReadFileCompleteDelegate::CreateUObject(this, &UOnlineHotfixManager::OnReadFileComplete);
+#if !UE_BUILD_SHIPPING
+	bLogMountedPakContents = FParse::Param(FCommandLine::Get(), TEXT("LogHotfixPakContents"));
+#endif
+	GameContentPath = FString() / FApp::GetGameName() / TEXT("Content");
+}
+
+UOnlineHotfixManager* UOnlineHotfixManager::Get(UWorld* World)
+{
+	UOnlineHotfixManager* DefaultObject = UOnlineHotfixManager::StaticClass()->GetDefaultObject<UOnlineHotfixManager>();
+	IOnlineSubsystem* OnlineSub = Online::GetSubsystem(World, DefaultObject->OSSName.Len() > 0 ? FName(*DefaultObject->OSSName) : NAME_None);
+	if (OnlineSub != nullptr)
+	{
+		UOnlineHotfixManager* HotfixManager = Cast<UOnlineHotfixManager>(OnlineSub->GetNamedInterface(NAME_HotfixManager));
+		if (HotfixManager == nullptr)
+		{
+			FString HotfixManagerClassName = DefaultObject->HotfixManagerClassName;
+			UClass* HotfixManagerClass = LoadClass<UOnlineHotfixManager>(nullptr, *HotfixManagerClassName, nullptr, LOAD_None, nullptr);
+			if (HotfixManagerClass == nullptr)
+			{
+				// Just use the default class if it couldn't load what was specified
+				HotfixManagerClass = UOnlineHotfixManager::StaticClass();
+			}
+			// Create it and store it
+			HotfixManager = NewObject<UOnlineHotfixManager>(GetTransientPackage(), HotfixManagerClass);
+			OnlineSub->SetNamedInterface(NAME_HotfixManager, HotfixManager);
+		}
+		return HotfixManager;
+	}
+	return nullptr;
+}
+
+void UOnlineHotfixManager::PostInitProperties()
+{
+#if !UE_BUILD_SHIPPING
+	FParse::Value(FCommandLine::Get(), TEXT("HOTFIXPREFIX="), DebugPrefix);
+#endif
+	// So we only try to apply files for this platform
+	PlatformPrefix = DebugPrefix + ANSI_TO_TCHAR(FPlatformProperties::PlatformName());
+	PlatformPrefix += TEXT("_");
+	// Build the default prefix too
+	DefaultPrefix = DebugPrefix + TEXT("Default");
+
+	Super::PostInitProperties();
+}
+
+void UOnlineHotfixManager::Init()
+{
+	bHotfixingInProgress = true;
+	bHotfixNeedsMapReload = false;
+	TotalFiles = 0;
+	NumDownloaded = 0;
+	TotalBytes = 0;
+	NumBytes = 0;
+	ChangedOrRemovedPakCount = 0;
+	// Build the name of the loc file that we'll care about
+	// It can change at runtime so build it just before fetching the data
+	GameLocName = DebugPrefix + FInternationalization::Get().GetCurrentCulture()->GetTwoLetterISOLanguageName() + TEXT("_Game.locres");
+	OnlineTitleFile = Online::GetTitleFileInterface(OSSName.Len() ? FName(*OSSName, FNAME_Find) : NAME_None);
+	if (OnlineTitleFile.IsValid())
+	{
+		OnEnumerateFilesCompleteDelegateHandle = OnlineTitleFile->AddOnEnumerateFilesCompleteDelegate_Handle(OnEnumerateFilesCompleteDelegate);
+		OnReadFileProgressDelegateHandle = OnlineTitleFile->AddOnReadFileProgressDelegate_Handle(OnReadFileProgressDelegate);
+		OnReadFileCompleteDelegateHandle = OnlineTitleFile->AddOnReadFileCompleteDelegate_Handle(OnReadFileCompleteDelegate);
+	}
+}
+
+void UOnlineHotfixManager::Cleanup()
+{
+	PendingHotfixFiles.Empty();
+	if (OnlineTitleFile.IsValid())
+	{
+		// Make sure to give back the memory used when reading the hotfix files
+		OnlineTitleFile->ClearFiles();
+		OnlineTitleFile->ClearOnEnumerateFilesCompleteDelegate_Handle(OnEnumerateFilesCompleteDelegateHandle);
+		OnlineTitleFile->ClearOnReadFileProgressDelegate_Handle(OnReadFileProgressDelegateHandle);
+		OnlineTitleFile->ClearOnReadFileCompleteDelegate_Handle(OnReadFileCompleteDelegateHandle);
+	}
+	OnlineTitleFile = nullptr;
+	bHotfixingInProgress = false;
+}
+
+void UOnlineHotfixManager::StartHotfixProcess()
+{
+	// Patching the editor this way seems like a bad idea
+	bool bShouldHotfix = IsRunningGame() || IsRunningDedicatedServer() || IsRunningClientOnly();
+	if (!bShouldHotfix)
+	{
+		UE_LOG(LogHotfixManager, Warning, TEXT("Hotfixing skipped when not running game/server"));
+		TriggerHotfixComplete(EHotfixResult::SuccessNoChange);
+		return;
+	}
+
+	if (bHotfixingInProgress)
+	{
+		UE_LOG(LogHotfixManager, Warning, TEXT("Hotfixing already in progress"));
+		return;
+	}
+
+	Init();
+	// Kick off an enumeration of the files that are available to download
+	if (OnlineTitleFile.IsValid())
+	{
+		OnlineTitleFile->EnumerateFiles();
+	}
+	else
+	{
+		UE_LOG(LogHotfixManager, Error, TEXT("Failed to start the hotfixing process due to no OnlineTitleInterface present for OSS(%s)"), *OSSName);
+		TriggerHotfixComplete(EHotfixResult::Failed);
+	}
+}
+
+struct FHotfixFileSortPredicate
+{
+	struct FHotfixFileNameSortPredicate
+	{
+		const FString PlatformPrefix;
+		const FString DefaultPrefix;
+
+		FHotfixFileNameSortPredicate(const FString& InPlatformPrefix, const FString& InDefaultPrefix) :
+			PlatformPrefix(InPlatformPrefix),
+			DefaultPrefix(InDefaultPrefix)
+		{
+		}
+
+		bool Compare(const FString& A, const FString& B) const
+		{
+			bool bAIsIni = A.EndsWith(TEXT("INI"));
+			bool bBIsIni = B.EndsWith(TEXT("INI"));
+			if (bAIsIni && bBIsIni)
+			{
+				bool bAStartsWithDefault = A.StartsWith(DefaultPrefix);
+				bool bBStartsWithDefault = B.StartsWith(DefaultPrefix);
+				bool bAStartsWithPlatformName = A.StartsWith(PlatformPrefix);
+				bool bBStartsWithPlatformName = B.StartsWith(PlatformPrefix);
+				// Sort any file with Default in front of the Platform version so they apply in correct order
+				if (bAStartsWithDefault && bBStartsWithPlatformName)
+				{
+					return true;
+				}
+				else if (bAStartsWithPlatformName && bBStartsWithDefault)
+				{
+					return false;
+				}
+				// Sort any Default/Platform INIs before other INIs
+				else if ((bAStartsWithDefault || bAStartsWithPlatformName) && !(bBStartsWithDefault || bBStartsWithPlatformName))
+				{
+					return true;
+				}
+				else if (!(bAStartsWithDefault || bAStartsWithPlatformName) && (bBStartsWithDefault || bBStartsWithPlatformName))
+				{
+					return false;
+				}
+			}
+			// Sort all INI files before non-INI files
+			else if (bAIsIni && !bBIsIni)
+			{
+				return true;
+			}
+			else if (!bAIsIni && bBIsIni)
+			{
+				return false;
+			}
+			// Sort by the string order
+			return A < B;
+		}
+	};
+	FHotfixFileNameSortPredicate FileNameSorter;
+
+	FHotfixFileSortPredicate(const FString& InPlatformPrefix, const FString& InDefaultPrefix) :
+		FileNameSorter(InPlatformPrefix, InDefaultPrefix)
+	{
+	}
+
+	bool operator()(const FCloudFileHeader &A, const FCloudFileHeader &B) const
+	{
+		return FileNameSorter.Compare(A.FileName, B.FileName);
+	}
+
+	bool operator()(const FString& A, const FString& B) const
+	{
+		return FileNameSorter.Compare(FPaths::GetCleanFilename(A), FPaths::GetCleanFilename(B));
+	}
+};
+
+void UOnlineHotfixManager::OnEnumerateFilesComplete(bool bWasSuccessful)
+{
+	if (bWasSuccessful)
+	{
+		check(OnlineTitleFile.IsValid());
+		// Cache our current set so we can compare for differences
+		LastHotfixFileList = HotfixFileList;
+		HotfixFileList.Empty();
+		// Get the new header data
+		OnlineTitleFile->GetFileList(HotfixFileList);
+		FilterHotfixFiles();
+		// Reduce the set of work to just the files that changed since last run
+		BuildHotfixFileListDeltas();
+		// Sort after filtering so that the comparison below doesn't fail to different order from the server
+		ChangedHotfixFileList.Sort<FHotfixFileSortPredicate>(FHotfixFileSortPredicate(PlatformPrefix, DefaultPrefix));
+		// Perform any undo operations needed
+		if (ChangedHotfixFileList.Num() > 0 || RemovedHotfixFileList.Num() > 0)
+		{
+			RestoreBackupIniFiles();
+			UnmountHotfixFiles();
+		}
+		// Read any changed files
+		if (ChangedHotfixFileList.Num() > 0)
+		{
+			// Update our totals for our progress delegates
+			TotalFiles = ChangedHotfixFileList.Num();
+			for (auto& FileHeader : ChangedHotfixFileList)
+			{
+				TotalBytes += FileHeader.FileSize;
+			}
+			ReadHotfixFiles();
+		}
+		else
+		{
+			UE_LOG(LogHotfixManager, Display, TEXT("Returned hotfix data is the same as last application, skipping the apply phase"));
+			TriggerHotfixComplete(EHotfixResult::SuccessNoChange);
+		}
+	}
+	else
+	{
+		UE_LOG(LogHotfixManager, Error, TEXT("Enumeration of hotfix files failed"));
+		TriggerHotfixComplete(EHotfixResult::Failed);
+	}
+}
+
+void UOnlineHotfixManager::CheckAvailability(FOnHotfixAvailableComplete& InCompletionDelegate)
+{
+	// Checking for hotfixes in editor is not supported
+	bool bShouldHotfix = IsRunningGame() || IsRunningDedicatedServer() || IsRunningClientOnly();
+	if (!bShouldHotfix)
+	{
+		UE_LOG(LogHotfixManager, Warning, TEXT("Hotfixing availability skipped when not running game/server"));
+		InCompletionDelegate.ExecuteIfBound(EHotfixResult::SuccessNoChange);
+		return;
+	}
+
+	if (bHotfixingInProgress)
+	{
+		UE_LOG(LogHotfixManager, Warning, TEXT("Hotfixing availability skipped because hotfix in progress"));
+		InCompletionDelegate.ExecuteIfBound(EHotfixResult::Failed);
+		return;
+	}
+
+	OnlineTitleFile = Online::GetTitleFileInterface(OSSName.Len() ? FName(*OSSName, FNAME_Find) : NAME_None);
+
+	FOnEnumerateFilesCompleteDelegate OnEnumerateFilesForAvailabilityCompleteDelegate;
+	OnEnumerateFilesForAvailabilityCompleteDelegate.BindUObject(this, &UOnlineHotfixManager::OnEnumerateFilesForAvailabilityComplete, InCompletionDelegate);
+	OnEnumerateFilesForAvailabilityCompleteDelegateHandle = OnlineTitleFile->AddOnEnumerateFilesCompleteDelegate_Handle(OnEnumerateFilesForAvailabilityCompleteDelegate);
+
+	bHotfixingInProgress = true;
+
+	// Kick off an enumeration of the files that are available to download
+	if (OnlineTitleFile.IsValid())
+	{
+		OnlineTitleFile->EnumerateFiles();
+	}
+	else
+	{
+		UE_LOG(LogHotfixManager, Error, TEXT("Failed to start the hotfix check process due to no OnlineTitleInterface present for OSS(%s)"), *OSSName);
+		TriggerHotfixComplete(EHotfixResult::Failed);
+	}
+}
+
+void UOnlineHotfixManager::OnEnumerateFilesForAvailabilityComplete(bool bWasSuccessful, FOnHotfixAvailableComplete InCompletionDelegate)
+{
+	if (OnlineTitleFile.IsValid())
+	{
+		OnlineTitleFile->ClearOnEnumerateFilesCompleteDelegate_Handle(OnEnumerateFilesForAvailabilityCompleteDelegateHandle);
+	}
+	
+	EHotfixResult Result = EHotfixResult::Failed;
+	if (bWasSuccessful)
+	{
+		TArray<FCloudFileHeader> TmpHotfixFileList;
+		TArray<FCloudFileHeader> TmpLastHotfixFileList;
+
+		TmpHotfixFileList = HotfixFileList;
+		TmpLastHotfixFileList = LastHotfixFileList;
+
+		// Cache our current set so we can compare for differences
+		LastHotfixFileList = HotfixFileList;
+		HotfixFileList.Empty();
+		// Get the new header data
+		OnlineTitleFile->GetFileList(HotfixFileList);
+		FilterHotfixFiles();
+		// Reduce the set of work to just the files that changed since last run
+		BuildHotfixFileListDeltas();
+
+		// Read any changed files
+		if (ChangedHotfixFileList.Num() > 0 || RemovedHotfixFileList.Num() > 0)
+		{
+			UE_LOG(LogHotfixManager, Display, TEXT("Hotfix files available"));
+			Result = EHotfixResult::Success;
+		}
+		else
+		{
+			UE_LOG(LogHotfixManager, Display, TEXT("Returned hotfix data is the same as last application, returning nothing to do"));
+			Result = EHotfixResult::SuccessNoChange;
+		}
+
+		// Restore state to before the check
+		RemovedHotfixFileList.Empty();
+		ChangedHotfixFileList.Empty();
+		HotfixFileList = TmpHotfixFileList;
+		LastHotfixFileList = TmpLastHotfixFileList;
+	}
+	else
+	{
+		UE_LOG(LogHotfixManager, Error, TEXT("Enumeration of hotfix files failed"));
+	}
+
+	OnlineTitleFile = nullptr;
+	bHotfixingInProgress = false;
+	InCompletionDelegate.ExecuteIfBound(Result);
+}
+
+void UOnlineHotfixManager::BuildHotfixFileListDeltas()
+{
+	RemovedHotfixFileList.Empty();
+	ChangedHotfixFileList.Empty();
+	// Go through the current list and see if it's changed from the previous attempt
+	for (auto& CurrentHeader : HotfixFileList)
+	{
+		bool bFoundMatch = false;
+		for (auto& LastHeader : LastHotfixFileList)
+		{
+			if (LastHeader == CurrentHeader)
+			{
+				bFoundMatch = true;
+				break;
+			}
+		}
+		if (!bFoundMatch)
+		{
+			// We're different so add to the process list
+			ChangedHotfixFileList.Add(CurrentHeader);
+		}
+	}
+	// Find any files that have been removed from the set of hotfix files
+	for (auto& LastHeader : LastHotfixFileList)
+	{
+		bool bFoundMatch = false;
+		for (auto& CurrentHeader : HotfixFileList)
+		{
+			if (LastHeader.FileName == CurrentHeader.FileName)
+			{
+				bFoundMatch = true;
+				break;
+			}
+		}
+		if (!bFoundMatch)
+		{
+			// We've been removed so add to the removed list
+			RemovedHotfixFileList.Add(LastHeader);
+		}
+	}
+}
+
+void UOnlineHotfixManager::FilterHotfixFiles()
+{
+	for (int32 Idx = 0; Idx < HotfixFileList.Num(); Idx++)
+	{
+		if (!WantsHotfixProcessing(HotfixFileList[Idx]))
+		{
+			HotfixFileList.RemoveAt(Idx, 1, false);
+			Idx--;
+		}
+	}
+}
+
+void UOnlineHotfixManager::ReadHotfixFiles()
+{
+	if (ChangedHotfixFileList.Num())
+	{
+		check(OnlineTitleFile.IsValid());
+		// Kick off a read for each file
+		// Do this in two passes so already cached files don't trigger completion
+		for (auto& FileHeader : ChangedHotfixFileList)
+		{
+			UE_LOG(LogOnline, VeryVerbose, TEXT("HF: %s %s %d "), *FileHeader.DLName, *FileHeader.FileName, FileHeader.FileSize);
+			PendingHotfixFiles.Add(FileHeader.DLName, FPendingFileDLProgress());
+		}
+		for (auto& FileHeader : ChangedHotfixFileList)
+		{
+			OnlineTitleFile->ReadFile(FileHeader.DLName);
+		}
+	}
+	else
+	{
+		UE_LOG(LogHotfixManager, Display, TEXT("No hotfix files need to be downloaded"));
+		TriggerHotfixComplete(EHotfixResult::Success);
+	}
+}
+
+void UOnlineHotfixManager::OnReadFileComplete(bool bWasSuccessful, const FString& FileName)
+{
+	if (PendingHotfixFiles.Contains(FileName))
+	{
+		if (bWasSuccessful)
+		{
+			FCloudFileHeader* Header = GetFileHeaderFromDLName(FileName);
+			check(Header != nullptr);
+			UE_LOG(LogHotfixManager, Log, TEXT("Hotfix file (%s) downloaded. Size was (%d)"), *GetFriendlyNameFromDLName(FileName), Header->FileSize);
+			// Completion updates the file count and progress updates the byte count
+			UpdateProgress(1, 0);
+			PendingHotfixFiles.Remove(FileName);
+			if (PendingHotfixFiles.Num() == 0)
+			{
+				// All files have been downloaded so now apply the files
+				ApplyHotfix();
+			}
+		}
+		else
+		{
+			UE_LOG(LogHotfixManager, Error, TEXT("Hotfix file (%s) failed to download"), *GetFriendlyNameFromDLName(FileName));
+			TriggerHotfixComplete(EHotfixResult::Failed);
+		}
+	}
+}
+
+void UOnlineHotfixManager::UpdateProgress(uint32 FileCount, uint64 UpdateSize)
+{
+	NumDownloaded += FileCount;
+	NumBytes += UpdateSize;
+	// Update our progress
+	TriggerOnHotfixProgressDelegates(NumDownloaded, TotalFiles, NumBytes, TotalBytes);
+}
+
+void UOnlineHotfixManager::ApplyHotfix()
+{
+	for (auto& FileHeader : ChangedHotfixFileList)
+	{
+		if (!ApplyHotfixProcessing(FileHeader))
+		{
+			UE_LOG(LogHotfixManager, Error, TEXT("Couldn't apply hotfix file (%s)"), *FileHeader.FileName);
+			TriggerHotfixComplete(EHotfixResult::Failed);
+			return;
+		}
+		// Let anyone listening know we just processed this file
+		TriggerOnHotfixProcessedFileDelegates(FileHeader.FileName, GetCachedDirectory() / FileHeader.DLName);
+	}
+	UE_LOG(LogHotfixManager, Display, TEXT("Hotfix data has been successfully applied"));
+	EHotfixResult Result = EHotfixResult::Success;
+	if (ChangedOrRemovedPakCount > 0)
+	{
+		UE_LOG(LogHotfixManager, Display, TEXT("Hotfix has changed or removed PAK files so a relaunch of the app is needed"));
+		Result = EHotfixResult::SuccessNeedsRelaunch;
+	}
+	else if (bHotfixNeedsMapReload)
+	{
+		UE_LOG(LogHotfixManager, Display, TEXT("Hotfix has detected PAK files containing currently loaded maps, so a level load is needed"));
+		Result = EHotfixResult::SuccessNeedsReload;
+	}
+	TriggerHotfixComplete(Result);
+}
+
+void UOnlineHotfixManager::TriggerHotfixComplete(EHotfixResult HotfixResult)
+{
+	TriggerOnHotfixCompleteDelegates(HotfixResult);
+	if (HotfixResult == EHotfixResult::Failed)
+	{
+		HotfixFileList.Empty();
+		UnmountHotfixFiles();
+	}
+	Cleanup();
+}
+
+bool UOnlineHotfixManager::WantsHotfixProcessing(const FCloudFileHeader& FileHeader)
+{
+	const FString Extension = FPaths::GetExtension(FileHeader.FileName);
+	if (Extension == TEXT("INI"))
+	{
+		return FileHeader.FileName.StartsWith(PlatformPrefix) || FileHeader.FileName.StartsWith(DefaultPrefix);
+	}
+	else if (Extension == TEXT("PAK"))
+	{
+		return FileHeader.FileName.Find(PlatformPrefix) != -1;
+	}
+	return FileHeader.FileName == GameLocName;
+}
+
+bool UOnlineHotfixManager::ApplyHotfixProcessing(const FCloudFileHeader& FileHeader)
+{
+	bool bSuccess = false;
+	const FString Extension = FPaths::GetExtension(FileHeader.FileName);
+	if (Extension == TEXT("INI"))
+	{
+		TArray<uint8> FileData;
+		if (OnlineTitleFile->GetFileContents(FileHeader.DLName, FileData))
+		{
+			// Convert to a FString
+			FileData.Add(0);
+			const FString IniData = (ANSICHAR*)FileData.GetData();
+			bSuccess = HotfixIniFile(FileHeader.FileName, IniData);
+		}
+	}
+	else if (Extension == TEXT("LOCRES"))
+	{
+		HotfixLocFile(FileHeader);
+		// Currently no failure case for this
+		bSuccess = true;
+	}
+	else if (Extension == TEXT("PAK"))
+	{
+		bSuccess = HotfixPakFile(FileHeader);
+	}
+	OnlineTitleFile->ClearFile(FileHeader.FileName);
+	return bSuccess;
+}
+
+FString UOnlineHotfixManager::GetStrippedConfigFileName(const FString& IniName)
+{
+	FString StrippedIniName(IniName);
+	if (StrippedIniName.StartsWith(PlatformPrefix))
+	{
+		StrippedIniName = IniName.Right(StrippedIniName.Len() - PlatformPrefix.Len());
+	}
+	else if (StrippedIniName.StartsWith(DefaultPrefix))
+	{
+		StrippedIniName = IniName.Right(StrippedIniName.Len() - DefaultPrefix.Len());
+	}
+	else if (StrippedIniName.StartsWith(DebugPrefix))
+	{
+		StrippedIniName = IniName.Right(StrippedIniName.Len() - DebugPrefix.Len());
+	}
+	return StrippedIniName;
+}
+
+FString UOnlineHotfixManager::GetConfigFileNamePath(const FString& IniName)
+{
+	return FPaths::GeneratedConfigDir() + ANSI_TO_TCHAR(FPlatformProperties::PlatformName()) / IniName;
+}
+
+FConfigFile* UOnlineHotfixManager::GetConfigFile(const FString& IniName)
+{
+	FString StrippedIniName(GetStrippedConfigFileName(IniName));
+	FConfigFile* ConfigFile = nullptr;
+	// Look for the first matching INI file entry
+	for (TMap<FString, FConfigFile>::TIterator It(*GConfig); It; ++It)
+	{
+		if (It.Key().EndsWith(StrippedIniName))
+		{
+			ConfigFile = &It.Value();
+			break;
+		}
+	}
+	// If not found, add this file to the config cache
+	if (ConfigFile == nullptr)
+	{
+		const FString IniNameWithPath = GetConfigFileNamePath(StrippedIniName);
+		FConfigFile Empty;
+		GConfig->SetFile(IniNameWithPath, &Empty);
+		ConfigFile = GConfig->Find(IniNameWithPath, false);
+	}
+	check(ConfigFile);
+	// We never want to save these merged files
+	ConfigFile->NoSave = true;
+	return ConfigFile;
+}
+
+bool UOnlineHotfixManager::HotfixIniFile(const FString& FileName, const FString& IniData)
+{
+	FConfigFile* ConfigFile = GetConfigFile(FileName);
+	// Store the original file so we can undo this later
+	BackupIniFile(FileName, ConfigFile);
+	// Merge the string into the config file
+	ConfigFile->CombineFromBuffer(IniData);
+	TArray<UClass*> Classes;
+	TArray<UObject*> PerObjectConfigObjects;
+	int32 StartIndex = 0;
+	int32 EndIndex = 0;
+	// Find the set of object classes that were affected
+	while (StartIndex >= 0 && StartIndex < IniData.Len() && EndIndex >= StartIndex)
+	{
+		// Find the next section header
+		StartIndex = IniData.Find(TEXT("["), ESearchCase::IgnoreCase, ESearchDir::FromStart, StartIndex);
+		if (StartIndex > -1)
+		{
+			// Find the ending section identifier
+			EndIndex = IniData.Find(TEXT("]"), ESearchCase::IgnoreCase, ESearchDir::FromStart, StartIndex);
+			if (EndIndex > StartIndex)
+			{
+				int32 PerObjectNameIndex = IniData.Find(TEXT(" "), ESearchCase::IgnoreCase, ESearchDir::FromStart, StartIndex);
+				// Per object config entries will have a space in the name, but classes won't
+				if (PerObjectNameIndex == -1 || PerObjectNameIndex > EndIndex)
+				{
+					const TCHAR* ScriptHeader = TEXT("[/Script/");
+					if (FCString::Strnicmp(*IniData + StartIndex, ScriptHeader, FCString::Strlen(ScriptHeader)) == 0)
+					{
+						const int32 ScriptSectionTag = 9;
+						// Snip the text out and try to find the class for that
+						const FString PackageClassName = IniData.Mid(StartIndex + ScriptSectionTag, EndIndex - StartIndex - ScriptSectionTag);
+						// Find the class for this so we know what to update
+						UClass* Class = FindObject<UClass>(nullptr, *PackageClassName, true);
+						if (Class)
+						{
+							// Add this to the list to check against
+							Classes.Add(Class);
+						}
+					}
+				}
+				// Handle the per object config case by finding the object for reload
+				else
+				{
+					const int32 ClassNameStart = PerObjectNameIndex + 1;
+					const FString ClassName = IniData.Mid(ClassNameStart, EndIndex - ClassNameStart);
+
+					// Look up the class to search for
+					UClass* ObjectClass = FindObject<UClass>(ANY_PACKAGE, *ClassName);
+
+					if (ObjectClass)
+					{
+						const int32 Count = PerObjectNameIndex - StartIndex - 1;
+						const FString PerObjectName = IniData.Mid(StartIndex + 1, Count);
+
+						// Explicitly search the transient package (won't update non-transient objects)
+						UObject* PerObject = StaticFindObject(ObjectClass, ANY_PACKAGE, *PerObjectName, false);
+						if (PerObject != nullptr)
+						{
+							PerObjectConfigObjects.Add(PerObject);
+						}
+					}
+					else
+					{
+						UE_LOG(LogHotfixManager, Warning, TEXT("Specified per-object class %s was not found"), *ClassName);
+					}
+				}
+				StartIndex = EndIndex;
+			}
+		}
+	}
+
+	int32 NumObjectsReloaded = 0;
+	const double StartTime = FPlatformTime::Seconds();
+	if (Classes.Num())
+	{
+		// Now that we have a list of classes to update, we can iterate objects and reload
+		for (FObjectIterator It; It; ++It)
+		{
+			UClass* Class = It->GetClass();
+			if (Class->HasAnyClassFlags(CLASS_Config))
+			{
+				// Check to see if this class is in our list (yes, potentially n^2, but not in practice)
+				for (int32 ClassIndex = 0; ClassIndex < Classes.Num(); ClassIndex++)
+				{
+					if (It->IsA(Classes[ClassIndex]))
+					{
+						// Force a reload of the config vars
+						It->ReloadConfig();
+						NumObjectsReloaded++;
+						break;
+					}
+				}
+			}
+		}
+	}
+	// Reload any PerObjectConfig objects that were affected
+	for (auto ReloadObject : PerObjectConfigObjects)
+	{
+		ReloadObject->ReloadConfig();
+		NumObjectsReloaded++;
+	}
+	UE_LOG(LogHotfixManager, Log, TEXT("Updating config from %s took %f seconds and reloaded %d objects"),
+		*FileName, FPlatformTime::Seconds() - StartTime, NumObjectsReloaded);
+	return true;
+}
+
+void UOnlineHotfixManager::HotfixLocFile(const FCloudFileHeader& FileHeader)
+{
+	const double StartTime = FPlatformTime::Seconds();
+	FString LocFilePath = FString::Printf(TEXT("%s/%s"), *GetCachedDirectory(), *FileHeader.DLName);
+	FTextLocalizationManager::Get().UpdateFromLocalizationResource(LocFilePath);
+	UE_LOG(LogHotfixManager, Log, TEXT("Updating loc from %s took %f seconds"), *FileHeader.FileName, FPlatformTime::Seconds() - StartTime);
+}
+
+bool UOnlineHotfixManager::HotfixPakFile(const FCloudFileHeader& FileHeader)
+{
+	if (!FCoreDelegates::OnMountPak.IsBound())
+	{
+		UE_LOG(LogHotfixManager, Error, TEXT("PAK file (%s) could not be mounted because OnMountPak is not bound"), *FileHeader.FileName);
+		return false;
+	}
+	FString PakLocation = FString::Printf(TEXT("%s/%s"), *GetCachedDirectory(), *FileHeader.DLName);
+	FPakFileVisitor Visitor;
+	if (FCoreDelegates::OnMountPak.Execute(PakLocation, 0, &Visitor))
+	{
+		MountedPakFiles.Add(FileHeader.DLName);
+		UE_LOG(LogHotfixManager, Log, TEXT("Hotfix mounted PAK file (%s)"), *FileHeader.FileName);
+		int32 NumInisReloaded = 0;
+		const double StartTime = FPlatformTime::Seconds();
+		TArray<FString> IniList;
+		// Iterate through the pak file's contents for INI and asset reloading
+		for (auto& InternalPakFileName : Visitor.Files)
+		{
+			if (InternalPakFileName.EndsWith(TEXT(".ini")))
+			{
+				IniList.Add(InternalPakFileName);
+			}
+			else if (!bHotfixNeedsMapReload && InternalPakFileName.EndsWith(FPackageName::GetMapPackageExtension()))
+			{
+				bHotfixNeedsMapReload = IsMapLoaded(InternalPakFileName);
+			}
+		}
+		// Sort the INIs so they are processed consistently
+		IniList.Sort<FHotfixFileSortPredicate>(FHotfixFileSortPredicate(PlatformPrefix, DefaultPrefix));
+		// Now process the INIs in sorted order
+		for (auto& IniName : IniList)
+		{
+			HotfixPakIniFile(IniName);
+			NumInisReloaded++;
+		}
+		UE_LOG(LogHotfixManager, Log, TEXT("Processing pak file (%s) took %f seconds and resulted in (%d) INIs being reloaded"),
+			*FileHeader.FileName, FPlatformTime::Seconds() - StartTime, NumInisReloaded);
+#if !UE_BUILD_SHIPPING
+		if (bLogMountedPakContents)
+		{
+			UE_LOG(LogHotfixManager, Log, TEXT("Files in pak file (%s):"), *FileHeader.FileName);
+			for (auto& FileName : Visitor.Files)
+			{
+				UE_LOG(LogHotfixManager, Log, TEXT("\t\t%s"), *FileName);
+			}
+		}
+#endif
+		return true;
+	}
+	return false;
+}
+
+bool UOnlineHotfixManager::IsMapLoaded(const FString& MapName)
+{
+	FString MapPackageName(MapName.Left(MapName.Len() - 5));
+	MapPackageName = MapPackageName.Replace(*GameContentPath, TEXT("/Game"));
+	// If this map's UPackage exists, it is currently loaded
+	UPackage* MapPackage = FindObject<UPackage>(ANY_PACKAGE, *MapPackageName, true);
+	return MapPackage != nullptr;
+}
+
+bool UOnlineHotfixManager::HotfixPakIniFile(const FString& FileName)
+{
+	FString StrippedName;
+	const double StartTime = FPlatformTime::Seconds();
+	// Need to strip off the PAK path
+	FileName.Split(TEXT("/"), nullptr, &StrippedName, ESearchCase::IgnoreCase, ESearchDir::FromEnd);
+	FConfigFile* ConfigFile = GetConfigFile(StrippedName);
+	if (!ConfigFile->Combine(FString(TEXT("../../../")) + FileName.Replace(*GameContentPath, TEXT("/Game"))))
+	{
+		UE_LOG(LogHotfixManager, Log, TEXT("Hotfix failed to merge INI (%s) found in a PAK file"), *FileName);
+		return false;
+	}
+	UE_LOG(LogHotfixManager, Log, TEXT("Hotfix merged INI (%s) found in a PAK file"), *FileName);
+	int32 NumObjectsReloaded = 0;
+	// Now that we have a list of classes to update, we can iterate objects and
+	// reload if they match the INI file that was changed
+	for (FObjectIterator It; It; ++It)
+	{
+		UClass* Class = It->GetClass();
+		if (Class->HasAnyClassFlags(CLASS_Config) &&
+			Class->ClassConfigName == ConfigFile->Name)
+		{
+			// Force a reload of the config vars
+			It->ReloadConfig();
+			NumObjectsReloaded++;
+		}
+	}
+	UE_LOG(LogHotfixManager, Log, TEXT("Updating config from %s took %f seconds reloading %d objects"),
+		*FileName, FPlatformTime::Seconds() - StartTime, NumObjectsReloaded);
+	return true;
+}
+
+const FString UOnlineHotfixManager::GetFriendlyNameFromDLName(const FString& DLName) const
+{
+	for (auto& Header : HotfixFileList)
+	{
+		if (Header.DLName == DLName)
+		{
+			return Header.FileName;
+		}
+	}
+	return FString();
+}
+
+void UOnlineHotfixManager::UnmountHotfixFiles()
+{
+	if (MountedPakFiles.Num() == 0)
+	{
+		return;
+	}
+	// Unmount any changed hotfix files since we need to download them again
+	for (auto& FileHeader : ChangedHotfixFileList)
+	{
+		for (int32 Index = 0; Index < MountedPakFiles.Num(); Index++)
+		{
+			if (MountedPakFiles[Index] == FileHeader.DLName)
+			{
+				FCoreDelegates::OnUnmountPak.Execute(MountedPakFiles[Index]);
+				MountedPakFiles.RemoveAt(Index);
+				ChangedOrRemovedPakCount++;
+				UE_LOG(LogHotfixManager, Log, TEXT("Hotfix unmounted PAK file (%s) so it can be redownloaded"), *FileHeader.FileName);
+				break;
+			}
+		}
+	}
+	// Unmount any removed hotfix files
+	for (auto& FileHeader : RemovedHotfixFileList)
+	{
+		for (int32 Index = 0; Index < MountedPakFiles.Num(); Index++)
+		{
+			if (MountedPakFiles[Index] == FileHeader.DLName)
+			{
+				FCoreDelegates::OnUnmountPak.Execute(MountedPakFiles[Index]);
+				MountedPakFiles.RemoveAt(Index);
+				ChangedOrRemovedPakCount++;
+				UE_LOG(LogHotfixManager, Log, TEXT("Hotfix unmounted PAK file (%s) since it was removed from the hotfix set"), *FileHeader.FileName);
+				break;
+			}
+		}
+	}
+}
+
+FCloudFileHeader* UOnlineHotfixManager::GetFileHeaderFromDLName(const FString& FileName)
+{
+	for (int32 Index = 0; Index < HotfixFileList.Num(); Index++)
+	{
+		if (HotfixFileList[Index].DLName == FileName)
+		{
+			return &HotfixFileList[Index];
+		}
+	}
+	return nullptr;
+}
+
+void UOnlineHotfixManager::OnReadFileProgress(const FString& FileName, uint64 BytesRead)
+{
+	if (PendingHotfixFiles.Contains(FileName))
+	{
+		// Since the title file is reporting absolute numbers subtract out the last update so we can add a delta
+		uint64 Delta = BytesRead - PendingHotfixFiles[FileName].Progress;
+		PendingHotfixFiles[FileName].Progress = BytesRead;
+		// Completion updates the file count and progress updates the byte count
+		UpdateProgress(0, Delta);
+	}
+}
+
+void UOnlineHotfixManager::BackupIniFile(const FString& IniName, const FConfigFile* ConfigFile)
+{
+	int32 AddAt = IniBackups.AddDefaulted();
+	IniBackups[AddAt].IniName = GetConfigFileNamePath(GetStrippedConfigFileName(IniName));
+	IniBackups[AddAt].ConfigData = *ConfigFile;
+	// There's a lack of deep copy related to the SourceConfigFile so null it out
+	IniBackups[AddAt].ConfigData.SourceConfigFile = nullptr;
+}
+
+void UOnlineHotfixManager::RestoreBackupIniFiles()
+{
+	if (IniBackups.Num() == 0)
+	{
+		return;
+	}
+	const double StartTime = FPlatformTime::Seconds();
+	TArray<FString> RestoredInis;
+	// Restore any changed INI files and build a list of which ones changed for UObject reloading below
+	for (auto& FileHeader : ChangedHotfixFileList)
+	{
+		if (FileHeader.FileName.EndsWith(TEXT(".INI")))
+		{
+			const FString ProcessedName = GetConfigFileNamePath(GetStrippedConfigFileName(FileHeader.FileName));
+			for (int32 Index = 0; Index < IniBackups.Num(); Index++)
+			{
+				if (IniBackups[Index].IniName == ProcessedName)
+				{
+					GConfig->SetFile(IniBackups[Index].IniName, &IniBackups[Index].ConfigData);
+					IniBackups.RemoveAt(Index);
+					RestoredInis.Add(ProcessedName);
+					break;
+				}
+			}
+		}
+	}
+	// Also restore any files that were previously part of the hotfix and now are not
+	for (auto& FileHeader : RemovedHotfixFileList)
+	{
+		if (FileHeader.FileName.EndsWith(TEXT(".INI")))
+		{
+			const FString ProcessedName = GetConfigFileNamePath(GetStrippedConfigFileName(FileHeader.FileName));
+			for (int32 Index = 0; Index < IniBackups.Num(); Index++)
+			{
+				if (IniBackups[Index].IniName == ProcessedName)
+				{
+					GConfig->SetFile(IniBackups[Index].IniName, &IniBackups[Index].ConfigData);
+					IniBackups.RemoveAt(Index);
+					RestoredInis.Add(ProcessedName);
+					break;
+				}
+			}
+		}
+	}
+	// Finally reload any UObjects that were affected by this
+	uint32 NumObjectsReloaded = 0;
+	if (RestoredInis.Num() > 0)
+	{
+		for (FObjectIterator It; It; ++It)
+		{
+			UClass* Class = It->GetClass();
+			if (Class->HasAnyClassFlags(CLASS_Config))
+			{
+				for (int32 Index = 0; Index < RestoredInis.Num(); Index++)
+				{
+					if (Class->GetConfigName() == RestoredInis[Index])
+					{
+						It->ReloadConfig();
+						NumObjectsReloaded++;
+						break;
+					}
+				}
+			}
+		}
+	}
+	UE_LOG(LogHotfixManager, Log, TEXT("Restoring config for %d changed INI files took %f seconds reloading %d objects"),
+		RestoredInis.Num(), FPlatformTime::Seconds() - StartTime, NumObjectsReloaded);
+}
+
+struct FHotfixManagerExec :
+	public FSelfRegisteringExec
+{
+	virtual bool Exec(UWorld* InWorld, const TCHAR* Cmd, FOutputDevice& Ar) override
+	{
+		if (FParse::Command(&Cmd, TEXT("HOTFIX")))
+		{
+			UOnlineHotfixManager* HotfixManager = UOnlineHotfixManager::Get(InWorld);
+			if (HotfixManager != nullptr)
+			{
+				HotfixManager->StartHotfixProcess();
+			}
+			return true;
+		}
+		else if (FParse::Command(&Cmd, TEXT("TESTHOTFIXSORT")))
+		{
+			TArray<FCloudFileHeader> TestList;
+			FCloudFileHeader Header;
+			Header.FileName = TEXT("SomeRandom.ini");
+			TestList.Add(Header);
+			Header.FileName = TEXT("pakchunk1-PS4_P.pak");
+			TestList.Add(Header);
+			Header.FileName = TEXT("EN_Game.locres");
+			TestList.Add(Header);
+			Header.FileName = TEXT("DefaultGame.ini");
+			TestList.Add(Header);
+			Header.FileName = TEXT("PS4_DefaultEngine.ini");
+			TestList.Add(Header);
+			Header.FileName = TEXT("DefaultEngine.ini");
+			TestList.Add(Header);
+			Header.FileName = TEXT("pakchunk0-PS4_P.pak");
+			TestList.Add(Header);
+			Header.FileName = TEXT("PS4_DefaultGame.ini");
+			TestList.Add(Header);
+			Header.FileName = TEXT("AnotherRandom.ini");
+			TestList.Add(Header);
+			TestList.Sort<FHotfixFileSortPredicate>(FHotfixFileSortPredicate(TEXT("PS4_"), TEXT("Default")));
+
+			UE_LOG(LogHotfixManager, Log, TEXT("Hotfixing sort is:"));
+			for (auto& FileHeader : TestList)
+			{
+				UE_LOG(LogHotfixManager, Log, TEXT("\t%s"), *FileHeader.FileName);
+			}
+
+			TArray<FString> TestList2;
+			TestList2.Add(TEXT("SomeRandom.ini"));
+			TestList2.Add(TEXT("DefaultGame.ini"));
+			TestList2.Add(TEXT("PS4_DefaultEngine.ini"));
+			TestList2.Add(TEXT("DefaultEngine.ini"));
+			TestList2.Add(TEXT("PS4_DefaultGame.ini"));
+			TestList2.Add(TEXT("AnotherRandom.ini"));
+			TestList2.Sort<FHotfixFileSortPredicate>(FHotfixFileSortPredicate(TEXT("PS4_"), TEXT("Default")));
+
+			UE_LOG(LogHotfixManager, Log, TEXT("Hotfixing PAK INI file sort is:"));
+			for (auto& IniName : TestList2)
+			{
+				UE_LOG(LogHotfixManager, Log, TEXT("\t%s"), *IniName);
+			}
+			return true;
+		}
+		return false;
+	}
+};
+static FHotfixManagerExec HotfixManagerExec;
diff --git a/Engine/Plugins/Online/OnlineFramework/Source/Hotfix/Public/HotfixModule.h b/Engine/Plugins/Online/OnlineFramework/Source/Hotfix/Public/HotfixModule.h
new file mode 100644
--- /dev/null
+++ b/Engine/Plugins/Online/OnlineFramework/Source/Hotfix/Public/HotfixModule.h
@@ -0,0 +1,17 @@
+// Copyright 1998-2016 Epic Games, Inc. All Rights Reserved.
+
+#pragma once
+
+#include "ModuleManager.h"
+
+/**
+ * Module for Hotfix base implementation
+ */
+class FHotfixModule : 
+	public IModuleInterface
+{
+private:
+// IModuleInterface
+	virtual void StartupModule() override;
+	virtual void ShutdownModule() override;
+};
diff --git a/Engine/Plugins/Online/OnlineFramework/Source/Lobby/Private/LobbyPrivatePCH.h b/Engine/Plugins/Online/OnlineFramework/Source/Lobby/Private/LobbyPrivatePCH.h
new file mode 100644
--- /dev/null
+++ b/Engine/Plugins/Online/OnlineFramework/Source/Lobby/Private/LobbyPrivatePCH.h
@@ -0,0 +1,11 @@
+// Copyright 1998-2016 Epic Games, Inc. All Rights Reserved.
+
+#pragma once
+
+#include "Core.h"
+#include "CoreUObject.h"
+#include "Engine.h"
+#include "Net/UnrealNetwork.h"
+
+// Module includes
+#include "LobbyModule.h"
\ No newline at end of file
diff --git a/Engine/Plugins/Online/OnlineFramework/Source/Lobby/Public/Lobby.h b/Engine/Plugins/Online/OnlineFramework/Source/Lobby/Public/Lobby.h
new file mode 100644
--- /dev/null
+++ b/Engine/Plugins/Online/OnlineFramework/Source/Lobby/Public/Lobby.h
@@ -0,0 +1,6 @@
+// Copyright 1998-2016 Epic Games, Inc. All Rights Reserved.
+
+#pragma once
+
+// Module includes
+#include "LobbyModule.h"
diff --git a/Engine/Plugins/Online/OnlineFramework/Source/Party/Private/Party.cpp b/Engine/Plugins/Online/OnlineFramework/Source/Party/Private/Party.cpp
new file mode 100644
--- /dev/null
+++ b/Engine/Plugins/Online/OnlineFramework/Source/Party/Private/Party.cpp
@@ -0,0 +1,1269 @@
+// Copyright 1998-2016 Epic Games, Inc. All Rights Reserved.
+
+#include "PartyPrivatePCH.h"
+#include "Party.h"
+#include "PartyGameState.h"
+#include "Engine/GameInstance.h"
+
+#include "Online.h"
+#include "OnlineSubsystemUtils.h"
+
+#define LOCTEXT_NAMESPACE "Parties"
+
+UParty::UParty(const FObjectInitializer& ObjectInitializer) :
+	Super(ObjectInitializer),
+	bLeavingPersistentParty(false)
+{
+	if (!HasAnyFlags(RF_ClassDefaultObject))
+	{
+	}
+}
+
+void UParty::Init()
+{
+	UGameInstance* GameInstance = GetGameInstance();
+	check(GameInstance);
+	GameInstance->OnNotifyPreClientTravel().AddUObject(this, &ThisClass::NotifyPreClientTravel);
+
+	FCoreUObjectDelegates::PostLoadMap.AddUObject(this, &ThisClass::OnPostLoadMap);
+}
+
+void UParty::InitPIE()
+{
+	OnPostLoadMap();
+}
+
+void UParty::OnPostLoadMap()
+{
+	if (!HasAnyFlags(RF_ClassDefaultObject))
+	{
+		RegisterIdentityDelegates();
+		RegisterPartyDelegates();
+	}
+}
+
+void UParty::OnShutdown()
+{
+	for (const auto& PartyKeyValue : JoinedParties)
+	{
+		UPartyGameState* Party = PartyKeyValue.Value;
+		if (Party)
+		{
+			Party->OnShutdown();
+		}
+	}
+
+	JoinedParties.Empty();
+
+	if (!HasAnyFlags(RF_ClassDefaultObject))
+	{
+		FCoreUObjectDelegates::PostLoadMap.RemoveAll(this);
+
+		UnregisterIdentityDelegates();
+		UnregisterPartyDelegates();
+	}
+
+	UGameInstance* GameInstance = GetGameInstance();
+	if (GameInstance)
+	{
+		GameInstance->OnNotifyPreClientTravel().RemoveAll(this);
+	}
+}
+
+void UParty::OnLogoutComplete(int32 LocalUserNum, bool bWasSuccessful)
+{
+	if (JoinedParties.Num())
+	{
+		UE_LOG(LogParty, Log, TEXT("Party cleanup on logout"));
+		TArray<FOnlinePartyTypeId> PartiesToRemove;
+		for (const auto& PartyKeyValue : JoinedParties)
+		{
+			PartiesToRemove.Add(PartyKeyValue.Key);
+		}
+		for (const auto& PartyKey : PartiesToRemove)
+		{
+			UPartyGameState** FoundParty = JoinedParties.Find(PartyKey);
+			if (FoundParty)
+			{
+				UPartyGameState* Party = *FoundParty;
+				if (Party && Party->GetPartyId().IsValid())
+				{
+					if (Party->GetPartyId().IsValid())
+					{
+						UE_LOG(LogParty, Log, TEXT("[%s] Removed"), *(Party->GetPartyId()->ToDebugString()));
+					}
+					else
+					{
+						UE_LOG(LogParty, Log, TEXT("[%s] Removed - Invalid party Id"));
+					}
+					Party->HandleRemovedFromParty(EMemberExitedReason::Left);
+				}
+				JoinedParties.Remove(PartyKey);
+			}
+		}
+
+		ensure(JoinedParties.Num() == 0);
+		JoinedParties.Empty();
+	}
+}
+
+void UParty::OnLoginStatusChanged(int32 LocalUserNum, ELoginStatus::Type OldStatus, ELoginStatus::Type NewStatus, const FUniqueNetId& NewId)
+{
+	if (NewStatus == ELoginStatus::NotLoggedIn)
+	{
+		if (JoinedParties.Num())
+		{
+			UE_LOG(LogParty, Log, TEXT("Party cleanup on logout"));
+			TArray<FOnlinePartyTypeId> PartiesToRemove;
+			for (const auto& PartyKeyValue : JoinedParties)
+			{
+				PartiesToRemove.Add(PartyKeyValue.Key);
+			}
+			for (const auto& PartyKey : PartiesToRemove)
+			{
+				UPartyGameState** FoundParty = JoinedParties.Find(PartyKey);
+				if (FoundParty)
+				{
+					UPartyGameState* Party = *FoundParty;
+					if (Party)
+					{
+						TSharedPtr<const FOnlinePartyId> PartyId = Party->GetPartyId();
+						FString PartyIdString = PartyId.IsValid() ? PartyId->ToDebugString() : TEXT("");
+						UE_LOG(LogParty, Log, TEXT("[%s] Removed"), *PartyIdString);
+						Party->HandleRemovedFromParty(EMemberExitedReason::Left);
+					}
+					JoinedParties.Remove(PartyKey);
+				}
+			}
+
+			ensure(JoinedParties.Num() == 0);
+			JoinedParties.Empty();
+		}
+	}
+
+	ClearPendingPartyJoin();
+}
+
+void UParty::AddReferencedObjects(UObject* InThis, FReferenceCollector& Collector)
+{
+	Super::AddReferencedObjects(InThis, Collector);
+
+	UParty* This = CastChecked<UParty>(InThis);
+	check(This);
+
+	TArray<UPartyGameState*> Parties;
+	This->JoinedParties.GenerateValueArray(Parties);
+	Collector.AddReferencedObjects(Parties);
+}
+
+void UParty::RegisterIdentityDelegates()
+{
+	UWorld* World = GetWorld();
+	check(World);
+
+	IOnlineIdentityPtr IdentityInt = Online::GetIdentityInterface(World);
+	if (IdentityInt.IsValid())
+	{
+		// Unbind and then rebind
+		UnregisterIdentityDelegates();
+
+		FOnLogoutCompleteDelegate OnLogoutCompleteDelegate;
+		OnLogoutCompleteDelegate.BindUObject(this, &ThisClass::OnLogoutComplete);
+
+		FOnLoginStatusChangedDelegate OnLoginStatusChangedDelegate;
+		OnLoginStatusChangedDelegate.BindUObject(this, &ThisClass::OnLoginStatusChanged);
+
+		for (int32 LocalPlayerId = 0; LocalPlayerId < MAX_LOCAL_PLAYERS; LocalPlayerId++)
+		{
+			LogoutCompleteDelegateHandle[LocalPlayerId] = IdentityInt->AddOnLogoutCompleteDelegate_Handle(LocalPlayerId, OnLogoutCompleteDelegate);
+			LogoutStatusChangedDelegateHandle[LocalPlayerId] = IdentityInt->AddOnLoginStatusChangedDelegate_Handle(LocalPlayerId, OnLoginStatusChangedDelegate);
+		}
+	}
+}
+
+void UParty::UnregisterIdentityDelegates()
+{
+	UWorld* World = GetWorld();
+	IOnlineIdentityPtr IdentityInt = Online::GetIdentityInterface(World);
+	if (IdentityInt.IsValid())
+	{
+		for (int32 LocalPlayerId = 0; LocalPlayerId < MAX_LOCAL_PLAYERS; LocalPlayerId++)
+		{
+			if (LogoutCompleteDelegateHandle[LocalPlayerId].IsValid())
+			{
+				IdentityInt->ClearOnLogoutCompleteDelegate_Handle(LocalPlayerId, LogoutCompleteDelegateHandle[LocalPlayerId]);
+			}
+
+			if (LogoutStatusChangedDelegateHandle[LocalPlayerId].IsValid())
+			{
+				IdentityInt->ClearOnLoginStatusChangedDelegate_Handle(LocalPlayerId, LogoutStatusChangedDelegateHandle[LocalPlayerId]);
+			}
+		}
+	}
+}
+
+void UParty::RegisterPartyDelegates()
+{
+	UnregisterPartyDelegates();
+
+	UWorld* World = GetWorld();
+	check(World);
+
+	IOnlinePartyPtr PartyInt = Online::GetPartyInterface(World);
+	if (PartyInt.IsValid())
+	{
+		PartyConfigChangedDelegateHandle = PartyInt->AddOnPartyConfigChangedDelegate_Handle(FOnPartyConfigChangedDelegate::CreateUObject(this, &ThisClass::PartyConfigChangedInternal));
+		PartyMemberJoinedDelegateHandle = PartyInt->AddOnPartyMemberJoinedDelegate_Handle(FOnPartyMemberJoinedDelegate::CreateUObject(this, &ThisClass::PartyMemberJoinedInternal));
+		PartyDataReceivedDelegateHandle = PartyInt->AddOnPartyDataReceivedDelegate_Handle(FOnPartyDataReceivedDelegate::CreateUObject(this, &ThisClass::PartyDataReceivedInternal));
+		PartyMemberDataReceivedDelegateHandle = PartyInt->AddOnPartyMemberDataReceivedDelegate_Handle(FOnPartyMemberDataReceivedDelegate::CreateUObject(this, &ThisClass::PartyMemberDataReceivedInternal));
+		PartyJoinRequestReceivedDelegateHandle = PartyInt->AddOnPartyJoinRequestReceivedDelegate_Handle(FOnPartyJoinRequestReceivedDelegate::CreateUObject(this, &ThisClass::PartyJoinRequestReceivedInternal));
+		PartyMemberChangedDelegateHandle = PartyInt->AddOnPartyMemberChangedDelegate_Handle(FOnPartyMemberChangedDelegate::CreateUObject(this, &ThisClass::PartyMemberChangedInternal));
+		PartyMemberExitedDelegateHandle = PartyInt->AddOnPartyMemberExitedDelegate_Handle(FOnPartyMemberExitedDelegate::CreateUObject(this, &ThisClass::PartyMemberExitedInternal));
+		PartyPromotionLockoutChangedDelegateHandle = PartyInt->AddOnPartyPromotionLockoutChangedDelegate_Handle(FOnPartyPromotionLockoutChangedDelegate::CreateUObject(this, &ThisClass::PartyPromotionLockoutStateChangedInternal));
+		PartyExitedDelegateHandle = PartyInt->AddOnPartyExitedDelegate_Handle(FOnPartyExitedDelegate::CreateUObject(this, &ThisClass::PartyExitedInternal));
+	}
+}
+
+void UParty::UnregisterPartyDelegates()
+{
+	UWorld* World = GetWorld();
+	IOnlinePartyPtr PartyInt = Online::GetPartyInterface(World);
+	if (PartyInt.IsValid())
+	{
+		PartyInt->ClearOnPartyConfigChangedDelegate_Handle(PartyConfigChangedDelegateHandle);
+		PartyInt->ClearOnPartyMemberJoinedDelegate_Handle(PartyMemberJoinedDelegateHandle);
+		PartyInt->ClearOnPartyDataReceivedDelegate_Handle(PartyDataReceivedDelegateHandle);
+		PartyInt->ClearOnPartyMemberDataReceivedDelegate_Handle(PartyMemberDataReceivedDelegateHandle);
+		PartyInt->ClearOnPartyJoinRequestReceivedDelegate_Handle(PartyJoinRequestReceivedDelegateHandle);
+		PartyInt->ClearOnPartyMemberChangedDelegate_Handle(PartyMemberChangedDelegateHandle);
+		PartyInt->ClearOnPartyMemberExitedDelegate_Handle(PartyMemberExitedDelegateHandle);
+		PartyInt->ClearOnPartyPromotionLockoutChangedDelegate_Handle(PartyPromotionLockoutChangedDelegateHandle);
+		PartyInt->ClearOnPartyExitedDelegate_Handle(PartyExitedDelegateHandle);
+	}
+}
+
+void UParty::NotifyPreClientTravel(const FString& PendingURL, ETravelType TravelType, bool bIsSeamlessTravel)
+{
+	for (auto& JoinedParty : JoinedParties)
+	{
+		UPartyGameState* PartyState = JoinedParty.Value;
+		PartyState->PreClientTravel();
+	}
+}
+
+bool UParty::HasPendingPartyJoin()
+{
+	return PendingPartyJoin.IsValid();
+}
+
+UPartyGameState* UParty::GetParty(const FOnlinePartyId& InPartyId) const
+{
+	UPartyGameState* ResultPtr = nullptr;
+
+	for (auto& JoinedParty : JoinedParties)
+	{
+		TSharedPtr<const FOnlinePartyId> PartyId = JoinedParty.Value->GetPartyId();
+		if (*PartyId == InPartyId)
+		{
+			ResultPtr = JoinedParty.Value;
+			break;
+		}
+	}
+
+	return ResultPtr;
+}
+
+UPartyGameState* UParty::GetParty(const FOnlinePartyTypeId InPartyTypeId) const
+{
+	UPartyGameState* ResultPtr = nullptr;
+
+	UPartyGameState* const * PartyPtr = JoinedParties.Find(InPartyTypeId);
+	ResultPtr = PartyPtr ? *PartyPtr : nullptr;
+	
+	return ResultPtr;
+}
+
+UPartyGameState* UParty::GetPersistentParty() const
+{
+	UPartyGameState* ResultPtr = nullptr;
+
+	UPartyGameState* const * PartyPtr = JoinedParties.Find(IOnlinePartySystem::GetPrimaryPartyTypeId());
+	ResultPtr = PartyPtr ? *PartyPtr : nullptr;
+
+	return ResultPtr;
+}
+
+void UParty::PartyConfigChangedInternal(const FUniqueNetId& InLocalUserId, const FOnlinePartyId& InPartyId, const TSharedRef<FPartyConfiguration>& InPartyConfig)
+{
+	UE_LOG(LogParty, Log, TEXT("[%s] Party config changed"), *InPartyId.ToString());
+
+	UPartyGameState* PartyState = GetParty(InPartyId);
+	if (PartyState)
+	{
+		PartyState->HandlePartyConfigChanged(InPartyConfig);
+	}
+	else
+	{
+		UE_LOG(LogParty, Warning, TEXT("[%s]: Missing party state during config change"), *InPartyId.ToString());
+	}
+}
+
+void UParty::PartyMemberJoinedInternal(const FUniqueNetId& InLocalUserId, const FOnlinePartyId& InPartyId, const FUniqueNetId& InMemberId)
+{
+	UE_LOG(LogParty, Log, TEXT("[%s] Player %s joined"), *InPartyId.ToString(), *InMemberId.ToString());
+
+	UPartyGameState* PartyState = GetParty(InPartyId);
+	if (PartyState)
+	{
+		PartyState->HandlePartyMemberJoined(InMemberId);
+	}
+	else
+	{
+		UE_LOG(LogParty, Warning, TEXT("[%s]: Missing party state player: %s"), *InPartyId.ToString(), *InMemberId.ToString());
+	}
+}
+
+void UParty::PartyDataReceivedInternal(const FUniqueNetId& InLocalUserId, const FOnlinePartyId& InPartyId, const TSharedRef<FOnlinePartyData>& InPartyData)
+{
+	UE_LOG(LogParty, Log, TEXT("[%s] party data received"), *InPartyId.ToString());
+
+	UPartyGameState* PartyState = GetParty(InPartyId);
+	if (PartyState)
+	{
+		PartyState->HandlePartyDataReceived(InPartyData);
+	}
+	else
+	{
+		UE_LOG(LogParty, Warning, TEXT("[%s]: Missing party state to apply data."), *InPartyId.ToString());
+	}
+}
+
+void UParty::PartyMemberDataReceivedInternal(const FUniqueNetId& InLocalUserId, const FOnlinePartyId& InPartyId, const FUniqueNetId& InMemberId, const TSharedRef<FOnlinePartyData>& InPartyMemberData)
+{
+	UE_LOG(LogParty, Log, TEXT("[%s] Player %s data received"), *InPartyId.ToString(), *InMemberId.ToString());
+
+	UPartyGameState* PartyState = GetParty(InPartyId);
+	if (PartyState)
+	{
+		PartyState->HandlePartyMemberDataReceived(InMemberId, InPartyMemberData);
+	}
+	else
+	{
+		UE_LOG(LogParty, Warning, TEXT("[%s]: Missing party state to apply data for player %s"), *InPartyId.ToString(), *InMemberId.ToString());
+	}
+}
+
+void UParty::PartyJoinRequestReceivedInternal(const FUniqueNetId& InLocalUserId, const FOnlinePartyId& InPartyId, const FUniqueNetId& SenderId)
+{
+	UPartyGameState* PartyState = GetParty(InPartyId);
+	if (PartyState)
+	{
+		PartyState->HandlePartyJoinRequestReceived(InLocalUserId, SenderId);
+	}
+	else
+	{
+		UE_LOG(LogParty, Warning, TEXT("[%s]: Missing party state to process join request."), *InPartyId.ToString());
+	}
+}
+
+void UParty::PartyMemberChangedInternal(const FUniqueNetId& InLocalUserId, const FOnlinePartyId& InPartyId, const FUniqueNetId& InMemberId, const EMemberChangedReason InReason)
+{
+	if (InLocalUserId == InMemberId)
+	{
+		UE_LOG(LogParty, Log, TEXT("[%s] [%s] local member %s"), *InPartyId.ToString(), *InMemberId.ToString(), ToString(InReason));
+	}
+	else
+	{
+		UE_LOG(LogParty, Log, TEXT("[%s] [%s] remote member %s"), *InPartyId.ToString(), *InMemberId.ToString(), ToString(InReason));
+	}
+
+	UPartyGameState* PartyState = GetParty(InPartyId);
+	if (!PartyState)
+	{
+		UE_LOG(LogParty, Warning, TEXT("[%s]: Missing party state during member change."), *InPartyId.ToString());
+	}
+
+	if (InReason == EMemberChangedReason::Disconnected)
+	{
+		// basically cleanup because they get kicked for now as long as party config says
+		// PartyConfig.bShouldRemoveOnDisconnection = true
+		// handled by PartyMemberLeft 
+	}
+	else if (InReason == EMemberChangedReason::Rejoined)
+	{
+		// can't happen as long as 
+		// PartyConfig.bShouldRemoveOnDisconnection = true
+	}
+	else if (InReason == EMemberChangedReason::Promoted)
+	{
+		if (PartyState)
+		{
+			PartyState->HandlePartyMemberPromoted(InMemberId);
+		}
+
+		if (PersistentPartyId.IsValid() &&
+			(InPartyId == *PersistentPartyId))
+		{
+			FUniqueNetIdRepl NewPartyLeader;
+			NewPartyLeader.SetUniqueNetId(InMemberId.AsShared());
+			UpdatePersistentPartyLeader(NewPartyLeader);
+		}
+	}
+}
+
+void UParty::PartyMemberExitedInternal(const FUniqueNetId& InLocalUserId, const FOnlinePartyId& InPartyId, const FUniqueNetId& InMemberId, const EMemberExitedReason InReason)
+{
+	if (InLocalUserId == InMemberId)
+	{
+		UE_LOG(LogParty, Log, TEXT("[%s] [%s] local member removed. Reason: %s"), *InPartyId.ToString(), *InMemberId.ToString(), ToString(InReason));
+	}
+	else
+	{
+		UE_LOG(LogParty, Log, TEXT("[%s] [%s] remote member exited. Reason: %s"), *InPartyId.ToString(), *InMemberId.ToString(), ToString(InReason));
+	}
+
+	if (InLocalUserId == InMemberId)
+	{
+		if (InReason == EMemberExitedReason::Left)
+		{
+			// EMemberExitedReason::Left -> local player chose to leave, handled by leave completion delegate
+		}
+		else
+		{
+			UPartyGameState* PartyState = GetParty(InPartyId);
+
+			// EMemberExitedReason::Removed -> removed from other reasons beside kick (can't happen?)
+			// EMemberExitedReason::Disbanded -> leader kicked everyone
+			// EMemberExitedReason::Kicked -> party/leader kicked you
+			// EMemberExitedReason::Unknown -> bad but try to cleanup anyway
+			if (PartyState)
+			{
+				FOnlinePartyTypeId PartyTypeId = PartyState->GetPartyTypeId();
+				PartyState->HandleRemovedFromParty(InReason);
+				JoinedParties.Remove(PartyTypeId);
+			}
+			else
+			{
+				UE_LOG(LogParty, Warning, TEXT("[%s]: Missing party state during local player exit."), *InPartyId.ToString());
+			}
+		}
+
+		// If the removal was the persistent party, make sure we are in a good state
+		if (PersistentPartyId.IsValid() &&
+			(InPartyId == *PersistentPartyId))
+		{
+			RestorePersistentPartyState();
+		}
+	}
+	else
+	{
+		UPartyGameState* PartyState = GetParty(InPartyId);
+		if (!PartyState)
+		{
+			UE_LOG(LogParty, Warning, TEXT("[%s]: Missing party state during remote player exit."), *InPartyId.ToString());
+		}
+
+		if (PartyState)
+		{
+			// EMemberExitedReason::Left -> player chose to leave
+			// EMemberExitedReason::Removed -> player was removed forcibly (timeout from disconnect, etc)
+			// EMemberExitedReason::Disbanded -> leader kicked everyone (only leader should get this)
+			// EMemberExitedReason::Kicked -> party/leader kicked this player
+			PartyState->HandlePartyMemberLeft(InMemberId, InReason);
+		}
+	}
+}
+
+void UParty::PartyPromotionLockoutStateChangedInternal(const FUniqueNetId& LocalUserId, const FOnlinePartyId& InPartyId, const bool bLockoutState)
+{
+	UE_LOG(LogParty, Log, TEXT("[%s] party lockout state changed to %s"), *InPartyId.ToString(), bLockoutState ? TEXT("true") : TEXT("false"));
+
+	UPartyGameState* PartyState = GetParty(InPartyId);
+	if (PartyState)
+	{
+		PartyState->HandleLockoutPromotionStateChange(bLockoutState);
+	}
+	else
+	{
+		UE_LOG(LogParty, Warning, TEXT("[%s]: Missing party state during lockout call"), *InPartyId.ToString());
+	}
+}
+
+void UParty::PartyExitedInternal(const FUniqueNetId& LocalUserId, const FOnlinePartyId& InPartyId)
+{
+	UE_LOG(LogParty, Log, TEXT("[%s] exited party %s"), *InPartyId.ToString(), *(InPartyId.ToDebugString()));
+
+	UPartyGameState* PartyState = GetParty(InPartyId);
+	if (PartyState != nullptr && (PartyState->OssParty->PartyTypeId == IOnlinePartySystem::GetPrimaryPartyTypeId()))
+	{
+		TSharedRef<FPartyConfiguration> Config = PartyState->OssParty->Config;
+		PartyState->HandleLeavingParty();
+		PartyState->HandleRemovedFromParty(EMemberExitedReason::Left);
+		JoinedParties.Remove(IOnlinePartySystem::GetPrimaryPartyTypeId());
+
+		UPartyDelegates::FOnCreateUPartyComplete CompletionDelegate;
+		CompletionDelegate.BindUObject(this, &ThisClass::OnPersistentPartyExitedInternalCompleted);
+		CreatePartyInternal(LocalUserId, IOnlinePartySystem::GetPrimaryPartyTypeId(), *Config, CompletionDelegate);
+	}
+	else
+	{
+		UE_LOG(LogParty, Warning, TEXT("[%s]: Missing party state during exit"), *InPartyId.ToString());
+	}
+}
+
+void UParty::OnPersistentPartyExitedInternalCompleted(const FUniqueNetId& LocalUserId, const ECreatePartyCompletionResult Result)
+{
+	if (Result == ECreatePartyCompletionResult::Succeeded)
+	{
+		OnCreatePesistentPartyCompletedCommon(LocalUserId);
+
+		UPartyGameState** Party = JoinedParties.Find(IOnlinePartySystem::GetPrimaryPartyTypeId());
+
+		if (Party != nullptr)
+		{
+			PartyResetForFrontendDelegate.Broadcast(*Party);
+		}
+	}
+
+	if (Result != ECreatePartyCompletionResult::Succeeded)
+	{
+		UE_LOG(LogParty, Warning, TEXT("Error when attempting to recreate persistent party on reconnection error=%s"), ToString(Result));
+	}
+}
+
+void UParty::CreatePartyInternal(const FUniqueNetId& InUserId, const FOnlinePartyTypeId InPartyTypeId, const FPartyConfiguration& InPartyConfig, const UPartyDelegates::FOnCreateUPartyComplete& InCompletionDelegate)
+{
+	ECreatePartyCompletionResult Result = ECreatePartyCompletionResult::UnknownClientFailure;
+	FString ErrorMsg;
+
+	UWorld* World = GetWorld();
+	IOnlinePartyPtr PartyInt = Online::GetPartyInterface(World);
+	if (PartyInt.IsValid())
+	{
+		FOnCreatePartyComplete CompletionDelegate;
+		CompletionDelegate.BindUObject(this, &ThisClass::OnCreatePartyInternalComplete, InPartyTypeId, InCompletionDelegate);
+		PartyInt->CreateParty(InUserId, InPartyTypeId, InPartyConfig, CompletionDelegate);
+		Result = ECreatePartyCompletionResult::Succeeded;
+	}
+	else
+	{
+		Result = ECreatePartyCompletionResult::UnknownClientFailure;
+		ErrorMsg = FString::Printf(TEXT("No party interface during JoinParty()"));
+	}
+
+	if (Result != ECreatePartyCompletionResult::Succeeded)
+	{
+		UE_LOG(LogParty, Warning, TEXT("%s"), *ErrorMsg);
+		InCompletionDelegate.ExecuteIfBound(InUserId, Result);
+	}
+}
+
+void UParty::OnCreatePartyInternalComplete(const FUniqueNetId& LocalUserId, const TSharedPtr<const FOnlinePartyId>& InPartyId, const ECreatePartyCompletionResult Result, const FOnlinePartyTypeId InPartyTypeId, UPartyDelegates::FOnCreateUPartyComplete InCompletionDelegate)
+{
+	const FString PartyIdDebugString = InPartyId.IsValid() ? InPartyId->ToDebugString() : TEXT("Invalid");
+	UE_LOG(LogParty, Display, TEXT("OnCreatePartyInternalComplete() %s %s"), *PartyIdDebugString, ToString(Result));
+
+	ECreatePartyCompletionResult LocalResult = Result;
+	if (Result == ECreatePartyCompletionResult::Succeeded)
+	{
+		UWorld* World = GetWorld();
+		IOnlinePartyPtr PartyInt = Online::GetPartyInterface(World);
+		if (PartyInt.IsValid())
+		{
+			TSharedPtr<const FOnlineParty> Party = PartyInt->GetParty(LocalUserId, InPartyTypeId);
+			if (ensure(Party.IsValid()))
+			{
+				TSubclassOf<UPartyGameState>* PartyGameStateClass = PartyClasses.Find(InPartyTypeId);
+				if (PartyGameStateClass != nullptr)
+				{
+					UPartyGameState* NewParty = NewObject<UPartyGameState>(this, *PartyGameStateClass);
+
+					// Add right away so future delegate broadcasts have this available
+					JoinedParties.Add(InPartyTypeId, NewParty);
+
+					// Initialize and trigger delegates
+					NewParty->InitFromCreate(LocalUserId, Party);
+
+					LocalResult = ECreatePartyCompletionResult::Succeeded;
+				}
+				else
+				{
+					LocalResult = ECreatePartyCompletionResult::UnknownClientFailure;
+				}
+			}
+			else
+			{
+				LocalResult = ECreatePartyCompletionResult::UnknownClientFailure;
+			}
+		}
+		else
+		{
+			LocalResult = ECreatePartyCompletionResult::UnknownClientFailure;
+		}
+	}
+
+	if (LocalResult != ECreatePartyCompletionResult::Succeeded)
+	{
+		UE_LOG(LogParty, Warning, TEXT("Error when creating party %s error=%s"), *PartyIdDebugString, ToString(LocalResult));
+	}
+
+	InCompletionDelegate.ExecuteIfBound(LocalUserId, LocalResult);
+}
+
+void UParty::JoinPartyInternal(const FUniqueNetId& InUserId, const FPartyDetails& InPartyDetails, const UPartyDelegates::FOnJoinUPartyComplete& InCompletionDelegate)
+{
+	EJoinPartyCompletionResult Result = EJoinPartyCompletionResult::UnknownClientFailure;
+	FString ErrorMsg;
+
+	IOnlinePartyPtr PartyInt = Online::GetPartyInterface(GetWorld());
+	if (PartyInt.IsValid())
+	{
+		if (InPartyDetails.IsValid())
+		{
+			const FOnlinePartyId& PartyId = *InPartyDetails.GetPartyId();
+			// High level party data check
+			UPartyGameState* PartyState = GetParty(PartyId);
+			// Interface level party data check should not be out of sync
+			TSharedPtr<const FOnlineParty> Party = PartyInt->GetParty(InUserId, PartyId);
+			if (!PartyState)
+			{
+				if (!Party.IsValid())
+				{
+					FOnJoinPartyComplete CompletionDelegate;
+					CompletionDelegate.BindUObject(this, &ThisClass::OnJoinPartyInternalComplete, InPartyDetails.GetPartyTypeId(), InCompletionDelegate);
+					PartyInt->JoinParty(InUserId, *(InPartyDetails.PartyJoinInfo), CompletionDelegate);
+					Result = EJoinPartyCompletionResult::Succeeded;
+				}
+				else
+				{
+					Result = EJoinPartyCompletionResult::AlreadyJoiningParty;
+					ErrorMsg = FString::Printf(TEXT("Already joining party %s, not joining again."), *InPartyDetails.GetPartyId()->ToString());
+				}
+			}
+			else
+			{
+				Result = EJoinPartyCompletionResult::AlreadyInParty;
+				ErrorMsg = FString::Printf(TEXT("Already in party %s, not joining again."), *InPartyDetails.GetPartyId()->ToString());
+			}
+		}
+		else
+		{
+			Result = EJoinPartyCompletionResult::JoinInfoInvalid;
+			ErrorMsg = FString::Printf(TEXT("Invalid party details, cannot join. Details: %s"), *InPartyDetails.ToString());
+		}
+	}
+	else
+	{
+		Result = EJoinPartyCompletionResult::UnknownClientFailure;
+		ErrorMsg = FString::Printf(TEXT("No party interface during JoinParty()"));
+	}
+
+	if (Result != EJoinPartyCompletionResult::Succeeded)
+	{
+		UE_LOG(LogParty, Warning, TEXT("%s"), *ErrorMsg);
+		InCompletionDelegate.ExecuteIfBound(InUserId, Result, 0);
+	}
+}
+
+void UParty::OnJoinPartyInternalComplete(const FUniqueNetId& LocalUserId, const FOnlinePartyId& InPartyId, const EJoinPartyCompletionResult Result, int32 DeniedResultCode, const FOnlinePartyTypeId InPartyTypeId, UPartyDelegates::FOnJoinUPartyComplete CompletionDelegate)
+{
+	const FString PartyIdDebugString = InPartyId.ToDebugString();
+	UE_LOG(LogParty, Display, TEXT("OnJoinPartyInternalComplete() %s %s."), *PartyIdDebugString, ToString(Result));
+
+	EJoinPartyCompletionResult LocalResult = Result;
+	if (Result == EJoinPartyCompletionResult::Succeeded)
+	{
+		UWorld* World = GetWorld();
+		IOnlinePartyPtr PartyInt = Online::GetPartyInterface(World);
+		if (PartyInt.IsValid())
+		{
+			TSharedPtr<const FOnlineParty> Party = PartyInt->GetParty(LocalUserId, InPartyTypeId);
+			if (Party.IsValid())
+			{
+				TSubclassOf<UPartyGameState>* PartyGameStateClass = PartyClasses.Find(InPartyTypeId);
+				if (PartyGameStateClass != nullptr)
+				{
+					UPartyGameState* NewParty = NewObject<UPartyGameState>(this, *PartyGameStateClass);
+
+					// Add right away so future delegate broadcasts have this available
+					JoinedParties.Add(InPartyTypeId, NewParty);
+
+					// Initialize and trigger delegates
+					NewParty->InitFromJoin(LocalUserId, Party);
+				}
+				else
+				{
+					LocalResult = EJoinPartyCompletionResult::AlreadyInPartyOfSpecifiedType;
+				}
+			}
+			else
+			{
+				LocalResult = EJoinPartyCompletionResult::UnknownClientFailure;
+			}
+		}
+		else
+		{
+			LocalResult = EJoinPartyCompletionResult::UnknownClientFailure;
+		}
+	}
+
+	if (LocalResult != EJoinPartyCompletionResult::Succeeded)
+	{
+		UE_LOG(LogParty, Warning, TEXT("Error when joining party %s error=%s"), *PartyIdDebugString, ToString(LocalResult));
+	}
+	
+	int32 OutDeniedResultCode = (Result == EJoinPartyCompletionResult::NotApproved) ? DeniedResultCode : 0;
+	CompletionDelegate.ExecuteIfBound(LocalUserId, LocalResult, OutDeniedResultCode);
+}
+
+void UParty::LeavePartyInternal(const FUniqueNetId& InUserId, const FOnlinePartyTypeId InPartyTypeId, const UPartyDelegates::FOnLeaveUPartyComplete& InCompletionDelegate)
+{
+	ELeavePartyCompletionResult Result = ELeavePartyCompletionResult::UnknownClientFailure;
+	FString ErrorMsg;
+
+	IOnlinePartyPtr PartyInt = Online::GetPartyInterface(GetWorld());
+	if (PartyInt.IsValid())
+	{
+		// Get the party by type here (don't rely on interface structures here because they can be missing during disconnects)
+		UPartyGameState* PartyState = GetParty(InPartyTypeId);
+		if (PartyState)
+		{
+			PartyState->HandleLeavingParty();
+
+			TSharedPtr<const FOnlinePartyId> PartyId = PartyState->GetPartyId();
+			if (ensure(PartyId.IsValid()))
+			{
+				FOnLeavePartyComplete CompletionDelegate;
+				CompletionDelegate.BindUObject(this, &ThisClass::OnLeavePartyInternalComplete, InPartyTypeId, InCompletionDelegate);
+				PartyInt->LeaveParty(InUserId, *PartyId, CompletionDelegate);
+				Result = ELeavePartyCompletionResult::Succeeded;
+			}
+			else
+			{
+				// Manual cleanup here because we can't call the above delegate
+				PartyState->HandleRemovedFromParty(EMemberExitedReason::Left);
+				JoinedParties.Remove(InPartyTypeId);
+			}
+		}
+		else
+		{
+			Result = ELeavePartyCompletionResult::UnknownParty;
+			ErrorMsg = FString::Printf(TEXT("Party not found in LeaveParty()"));
+		}
+	}
+	else
+	{
+		Result = ELeavePartyCompletionResult::UnknownClientFailure;
+		ErrorMsg = FString::Printf(TEXT("No party interface during LeaveParty()"));
+	}
+
+	if (Result != ELeavePartyCompletionResult::Succeeded)
+	{
+		UE_LOG(LogParty, Warning, TEXT("%s"), *ErrorMsg);
+		InCompletionDelegate.ExecuteIfBound(InUserId, Result);
+	}
+}
+
+void UParty::OnLeavePartyInternalComplete(const FUniqueNetId& LocalUserId, const FOnlinePartyId& InPartyId, const ELeavePartyCompletionResult Result, const FOnlinePartyTypeId InPartyTypeId, UPartyDelegates::FOnLeaveUPartyComplete CompletionDelegate)
+{
+	const FString PartyIdDebugString = InPartyId.ToDebugString();
+	UE_LOG(LogParty, Display, TEXT("OnLeavePartyComplete() %s %s."), *PartyIdDebugString, ToString(Result));
+
+	UPartyGameState* PartyState = GetParty(InPartyTypeId);
+	if (PartyState)
+	{
+		PartyState->HandleRemovedFromParty(EMemberExitedReason::Left);
+		JoinedParties.Remove(InPartyTypeId);
+	}
+	else
+	{
+		UE_LOG(LogParty, Warning, TEXT("OnLeavePartyComplete: Missing party state %s"), *PartyIdDebugString);
+	}
+
+	CompletionDelegate.ExecuteIfBound(LocalUserId, Result);
+}
+
+void UParty::CreateParty(const FUniqueNetId& InUserId, const FOnlinePartyTypeId InPartyTypeId, const FPartyConfiguration& InPartyConfig, const UPartyDelegates::FOnCreateUPartyComplete& InCompletionDelegate)
+{
+	UPartyDelegates::FOnCreateUPartyComplete CompletionDelegate;
+	CompletionDelegate.BindUObject(this, &ThisClass::OnCreatePartyComplete, InPartyTypeId, InCompletionDelegate);
+	CreatePartyInternal(InUserId, InPartyTypeId, InPartyConfig, CompletionDelegate);
+}
+
+void UParty::OnCreatePartyComplete(const FUniqueNetId& LocalUserId, const ECreatePartyCompletionResult Result, const FOnlinePartyTypeId InPartyTypeId, UPartyDelegates::FOnCreateUPartyComplete InCompletionDelegate)
+{
+	UE_LOG(LogParty, Display, TEXT("OnCreatePartyComplete() type(0x%08x) %s"), InPartyTypeId.GetValue(), ToString(Result));
+
+	InCompletionDelegate.ExecuteIfBound(LocalUserId, Result);
+}
+
+void UParty::JoinParty(const FUniqueNetId& InUserId, const FPartyDetails& InPartyDetails, const UPartyDelegates::FOnJoinUPartyComplete& InCompletionDelegate)
+{
+	UPartyDelegates::FOnJoinUPartyComplete CompletionDelegate;
+	CompletionDelegate.BindUObject(this, &ThisClass::OnJoinPartyComplete, InPartyDetails.GetPartyTypeId(), InCompletionDelegate);
+	JoinPartyInternal(InUserId, InPartyDetails, CompletionDelegate);
+}
+
+void UParty::OnJoinPartyComplete(const FUniqueNetId& LocalUserId, EJoinPartyCompletionResult Result, int32 DeniedResultCode, const FOnlinePartyTypeId InPartyTypeId, UPartyDelegates::FOnJoinUPartyComplete InCompletionDelegate)
+{
+	UE_LOG(LogParty, Display, TEXT("OnJoinPartyComplete() type(0x%08x) %s"), InPartyTypeId.GetValue(), ToString(Result));
+
+	int32 OutDeniedResultCode = (Result == EJoinPartyCompletionResult::NotApproved) ? DeniedResultCode : 0;
+	InCompletionDelegate.ExecuteIfBound(LocalUserId, Result, OutDeniedResultCode);
+}
+
+void UParty::LeaveParty(const FUniqueNetId& InUserId, const FOnlinePartyTypeId InPartyTypeId, const UPartyDelegates::FOnLeaveUPartyComplete& InCompletionDelegate)
+{
+	UPartyDelegates::FOnLeaveUPartyComplete CompletionDelegate;
+	CompletionDelegate.BindUObject(this, &ThisClass::OnLeavePartyComplete, InPartyTypeId, InCompletionDelegate);
+	LeavePartyInternal(InUserId, InPartyTypeId, CompletionDelegate);
+}
+
+void UParty::OnLeavePartyComplete(const FUniqueNetId& LocalUserId, const ELeavePartyCompletionResult Result, const FOnlinePartyTypeId InPartyTypeId, UPartyDelegates::FOnLeaveUPartyComplete InCompletionDelegate)
+{
+	UE_LOG(LogParty, Display, TEXT("OnLeavePartyComplete() type(0x%08x) %s"), InPartyTypeId.GetValue(), ToString(Result));
+
+	InCompletionDelegate.ExecuteIfBound(LocalUserId, Result);
+}
+
+void UParty::GetDefaultPersistentPartySettings(EPartyType& PartyType, bool& bLeaderFriendsOnly, bool& bLeaderInvitesOnly, bool& bAllowInvites)
+{
+	PartyType = EPartyType::Public;
+	bLeaderInvitesOnly = false;
+	bLeaderFriendsOnly = false;
+	bAllowInvites = true;
+}
+
+void UParty::GetPersistentPartyConfiguration(FPartyConfiguration& PartyConfig)
+{
+	EPartyType PartyType = EPartyType::Public;
+	bool bLeaderInvitesOnly = false;
+	bool bLeaderFriendsOnly = false;
+	bool bAllowInvites = true;
+	GetDefaultPersistentPartySettings(PartyType, bLeaderFriendsOnly, bLeaderInvitesOnly, bAllowInvites);
+
+	bool bIsPrivate = (PartyType == EPartyType::Private);
+
+	PartySystemPermissions::EPresencePermissions PresencePermissions;
+	if (bLeaderFriendsOnly)
+	{
+		if (bIsPrivate)
+		{
+			PresencePermissions = PartySystemPermissions::FriendsInviteOnly;
+		}
+		else
+		{
+			PresencePermissions = PartySystemPermissions::FriendsOnly;
+		}
+	}
+	else
+	{
+		if (bIsPrivate)
+		{
+			PresencePermissions = PartySystemPermissions::PublicInviteOnly;
+		}
+		else
+		{
+			PresencePermissions = PartySystemPermissions::Public;
+		}
+	}
+
+	PartyConfig.JoinRequestAction = EJoinRequestAction::Manual;
+	PartyConfig.bIsAcceptingMembers = bIsPrivate ? false : true;
+	PartyConfig.bShouldRemoveOnDisconnection = true;
+	PartyConfig.PresencePermissions = PresencePermissions;
+	PartyConfig.InvitePermissions = bAllowInvites ? (bLeaderInvitesOnly ? PartySystemPermissions::EInvitePermissions::Leader : PartySystemPermissions::EInvitePermissions::Anyone) : PartySystemPermissions::EInvitePermissions::Noone;
+
+	PartyConfig.MaxMembers = DefaultMaxPartySize;
+}
+
+void UParty::CreatePersistentParty(const FUniqueNetId& InUserId, const UPartyDelegates::FOnCreateUPartyComplete& InCompletionDelegate)
+{
+	if (PersistentPartyId.IsValid())
+	{
+		UE_LOG(LogParty, Warning, TEXT("Existing persistent party %s found when creating a new one."), *PersistentPartyId->ToString());
+	}
+
+	PersistentPartyId = nullptr;
+
+	FPartyConfiguration PartyConfig;
+	GetPersistentPartyConfiguration(PartyConfig);
+
+	UPartyDelegates::FOnCreateUPartyComplete CompletionDelegate;
+	CompletionDelegate.BindUObject(this, &ThisClass::OnCreatePersistentPartyComplete, InCompletionDelegate);
+	CreatePartyInternal(InUserId, IOnlinePartySystem::GetPrimaryPartyTypeId(), PartyConfig, CompletionDelegate);
+}
+
+void UParty::OnCreatePersistentPartyComplete(const FUniqueNetId& LocalUserId, const ECreatePartyCompletionResult Result, UPartyDelegates::FOnCreateUPartyComplete CompletionDelegate)
+{
+	UE_LOG(LogParty, Display, TEXT("OnCreatePersistentPartyComplete() %s"), ToString(Result));
+
+	if (Result == ECreatePartyCompletionResult::Succeeded)
+	{
+		OnCreatePesistentPartyCompletedCommon(LocalUserId);
+	}
+
+	CompletionDelegate.ExecuteIfBound(LocalUserId, Result);
+}
+
+void UParty::OnCreatePesistentPartyCompletedCommon(const FUniqueNetId& LocalUserId)
+{
+	UWorld* World = GetWorld();
+	check(World);
+
+	IOnlinePartyPtr PartyInt = Online::GetPartyInterface(World);
+	if (ensure(PartyInt.IsValid()))
+	{
+		TSharedPtr<const FOnlineParty> Party = PartyInt->GetParty(LocalUserId, IOnlinePartySystem::GetPrimaryPartyTypeId());
+		if (ensure(Party.IsValid()))
+		{
+			PersistentPartyId = Party->PartyId;
+		}
+	}
+
+	ensure(PersistentPartyId.IsValid());
+	UPartyGameState* PersistentParty = GetPersistentParty();
+	ensure(PersistentParty != nullptr);
+
+	EPartyType PartyType = EPartyType::Public;
+	bool bLeaderInvitesOnly = false;
+	bool bLeaderFriendsOnly = false;
+	bool bAllowInvites = true;
+	GetDefaultPersistentPartySettings(PartyType, bLeaderFriendsOnly, bLeaderInvitesOnly, bAllowInvites);
+
+	PersistentParty->SetPartyType(PartyType, bLeaderFriendsOnly, bLeaderInvitesOnly);
+	PersistentParty->SetInvitesDisabled(!bAllowInvites);
+
+	TSharedPtr<const FUniqueNetId> PartyLeaderPtr = PersistentParty->GetPartyLeader();
+	ensure(PartyLeaderPtr.IsValid());
+
+	FUniqueNetIdRepl PartyLeader(PartyLeaderPtr);
+	UpdatePersistentPartyLeader(PartyLeader);
+}
+
+void UParty::JoinPersistentParty(const FUniqueNetId& InUserId, const FPartyDetails& InPartyDetails, const UPartyDelegates::FOnJoinUPartyComplete& InCompletionDelegate)
+{
+	if (PersistentPartyId.IsValid())
+	{
+		UE_LOG(LogParty, Warning, TEXT("Existing persistent party %s found when joining a new one."), *PersistentPartyId->ToString());
+	}
+
+	PersistentPartyId = nullptr;
+
+	UPartyDelegates::FOnJoinUPartyComplete CompletionDelegate;
+	CompletionDelegate.BindUObject(this, &ThisClass::OnJoinPersistentPartyComplete, InCompletionDelegate);
+	JoinPartyInternal(InUserId, InPartyDetails, CompletionDelegate);
+}
+
+void UParty::OnJoinPersistentPartyComplete(const FUniqueNetId& LocalUserId, EJoinPartyCompletionResult Result, int32 DeniedResultCode, UPartyDelegates::FOnJoinUPartyComplete CompletionDelegate)
+{
+	UWorld* World = GetWorld();
+	IOnlinePartyPtr PartyInt = Online::GetPartyInterface(World);
+	if (PartyInt.IsValid())
+	{
+		if (Result == EJoinPartyCompletionResult::Succeeded)
+		{
+			TSharedPtr<const FOnlineParty> Party = PartyInt->GetParty(LocalUserId, IOnlinePartySystem::GetPrimaryPartyTypeId());
+			if (ensure(Party.IsValid()))
+			{
+				PersistentPartyId = Party->PartyId;
+			}
+		}
+	}
+
+	FString PartyIdDebugString = PersistentPartyId.IsValid() ? PersistentPartyId->ToDebugString() : TEXT("Invalid");
+	UE_LOG(LogParty, Display, TEXT("OnJoinPersistentPartyComplete() %s %s"), *PartyIdDebugString, ToString(Result));
+
+	int32 OutDeniedResultCode = (Result == EJoinPartyCompletionResult::NotApproved) ? DeniedResultCode : 0;
+	CompletionDelegate.ExecuteIfBound(LocalUserId, Result, OutDeniedResultCode);
+
+	if (Result == EJoinPartyCompletionResult::Succeeded)
+	{
+		ensure(PersistentPartyId.IsValid());
+		UPartyGameState* PersistentParty = GetPersistentParty();
+		if (PersistentParty)
+		{
+			TSharedPtr<const FUniqueNetId> PartyLeaderPtr = PersistentParty->GetPartyLeader();
+			if (PartyLeaderPtr.IsValid())
+			{
+				FUniqueNetIdRepl PartyLeader(PartyLeaderPtr);
+				UpdatePersistentPartyLeader(PartyLeader);
+			}
+			else
+			{
+				UE_LOG(LogParty, Warning, TEXT("OnJoinPersistentPartyComplete [%s]: Failed to update party leader"), *PersistentPartyId->ToString());
+			}
+		}
+		else
+		{
+			UE_LOG(LogParty, Warning, TEXT("OnJoinPersistentPartyComplete [%s]: Failed to find party state object"), PersistentPartyId.IsValid() ? *PersistentPartyId->ToString() : TEXT("INVALID"));
+		}
+	}
+	else
+	{
+		if (Result != EJoinPartyCompletionResult::AlreadyJoiningParty)
+		{
+			if (World)
+			{
+				// Try to get back to a good state
+				FTimerDelegate HandleJoinPartyFailure;
+				HandleJoinPartyFailure.BindUObject(this, &ThisClass::HandleJoinPersistentPartyFailure);
+				World->GetTimerManager().SetTimerForNextTick(HandleJoinPartyFailure);
+			}
+		}
+		else
+		{
+			UE_LOG(LogParty, Verbose, TEXT("OnJoinPersistentPartyComplete [%s]: already joining party."), PersistentPartyId.IsValid() ? *PersistentPartyId->ToString() : TEXT("INVALID"));
+		}
+	}
+}
+
+void UParty::HandleJoinPersistentPartyFailure()
+{
+	RestorePersistentPartyState();
+}
+
+void UParty::LeavePersistentParty(const FUniqueNetId& InUserId, const UPartyDelegates::FOnLeaveUPartyComplete& InCompletionDelegate)
+{
+	if (PersistentPartyId.IsValid())
+	{
+		if (!bLeavingPersistentParty)
+		{
+			UE_LOG(LogParty, Verbose, TEXT("LeavePersistentParty"));
+
+			UPartyDelegates::FOnLeaveUPartyComplete CompletionDelegate;
+			CompletionDelegate.BindUObject(this, &ThisClass::OnLeavePersistentPartyComplete, InCompletionDelegate);
+
+			bLeavingPersistentParty = true;
+			LeavePartyInternal(InUserId, IOnlinePartySystem::GetPrimaryPartyTypeId(), CompletionDelegate);
+		}
+		else
+		{
+			LeavePartyCompleteDelegates.Add(InCompletionDelegate);
+
+			UE_LOG(LogParty, Verbose, TEXT("LeavePersistentParty %d extra delegates"), LeavePartyCompleteDelegates.Num());
+		}
+	}
+	else
+	{
+		UE_LOG(LogParty, Warning, TEXT("No party during LeavePersistentParty()"));
+		InCompletionDelegate.ExecuteIfBound(InUserId, ELeavePartyCompletionResult::UnknownParty);
+	}
+}
+
+void UParty::OnLeavePersistentPartyComplete(const FUniqueNetId& LocalUserId, const ELeavePartyCompletionResult Result, UPartyDelegates::FOnLeaveUPartyComplete CompletionDelegate)
+{
+	UE_LOG(LogParty, Display, TEXT("OnLeavePersistentPartyComplete() %s"), ToString(Result));
+
+	ensure(bLeavingPersistentParty);
+	bLeavingPersistentParty = false;
+	PersistentPartyId = nullptr;
+
+	CompletionDelegate.ExecuteIfBound(LocalUserId, Result);
+
+	TArray<UPartyDelegates::FOnLeaveUPartyComplete> DelegatesCopy = LeavePartyCompleteDelegates;
+	LeavePartyCompleteDelegates.Empty();
+
+	// fire delegates for any/all LeavePersistentParty calls made while this one was in flight
+	for (const UPartyDelegates::FOnLeaveUPartyComplete& ExtraDelegate : DelegatesCopy)
+	{
+		ExtraDelegate.ExecuteIfBound(LocalUserId, Result);
+	}
+}
+
+void UParty::RestorePersistentPartyState()
+{
+	if (!GIsRequestingExit)
+	{
+		if (!bLeavingPersistentParty)
+		{
+			UWorld* World = GetWorld();
+			IOnlinePartyPtr PartyInt = Online::GetPartyInterface(World);
+			if (PartyInt.IsValid())
+			{
+				UGameInstance* GameInstance = GetGameInstance();
+				check(GameInstance);
+
+				TSharedPtr<const FUniqueNetId> LocalUserId = GameInstance->GetPrimaryPlayerUniqueId();
+				check(LocalUserId.IsValid() && LocalUserId->IsValid());
+
+				UPartyGameState* PersistentParty = GetPersistentParty();
+
+				// Check for existing party and create a new one if necessary
+				bool bFoundExistingPersistentParty = PersistentParty ? true : false;
+				if (bFoundExistingPersistentParty)
+				{
+					// In a party already, make sure the UI is aware of its state
+					if (PersistentParty->ResetForFrontend())
+					{
+						OnPartyResetForFrontend().Broadcast(PersistentParty);
+					}
+					else
+					{
+						// There was an issue resetting the party, so leave 
+						LeaveAndRestorePersistentParty();
+					}
+				}
+				else
+				{
+					PersistentPartyId = nullptr;
+
+					// Create a new party
+					CreatePersistentParty(*LocalUserId);
+				}
+			}
+		}
+		else
+		{
+			UE_LOG(LogParty, Verbose, TEXT("Can't RestorePersistentPartyState while leaving party, ignoring"));
+		}
+	}
+}
+
+void UParty::UpdatePersistentPartyLeader(const FUniqueNetIdRepl& NewPartyLeader)
+{
+}
+
+bool UParty::IsInParty(TSharedPtr<const FOnlinePartyId>& PartyId) const
+{
+	bool bFoundParty = false;
+
+	UWorld* World = GetWorld();
+	check(World);
+
+	IOnlinePartyPtr PartyInt = Online::GetPartyInterface(World);
+	if (PartyInt.IsValid())
+	{
+		UGameInstance* GameInstance = GetGameInstance();
+		check(GameInstance);
+
+		TSharedPtr<const FUniqueNetId> LocalUserId = GameInstance->GetPrimaryPlayerUniqueId();
+		if (ensure(LocalUserId.IsValid() && LocalUserId->IsValid()))
+		{
+			TArray<TSharedRef<const FOnlinePartyId>> LocalJoinedParties;
+			PartyInt->GetJoinedParties(*LocalUserId, LocalJoinedParties);
+			for (const auto& JoinedParty : LocalJoinedParties)
+			{
+				if (*JoinedParty == *PartyId)
+				{
+					bFoundParty = true;
+					break;
+				}
+			}
+		}
+	}
+
+	return bFoundParty;
+}
+
+void UParty::KickFromPersistentParty(const UPartyDelegates::FOnLeaveUPartyComplete& InCompletionDelegate)
+{
+	TSharedPtr<const FOnlinePartyId> LocalPersistentPartyId = GetPersistentPartyId();
+	UPartyGameState* PersistentParty = GetPersistentParty();
+	if (LocalPersistentPartyId.IsValid() && PersistentParty)
+	{
+		if (PersistentParty->GetPartySize() > 1)
+		{
+			UGameInstance* GameInstance = GetGameInstance();
+			check(GameInstance);
+
+			TSharedPtr<const FUniqueNetId> LocalUserId = GameInstance->GetPrimaryPlayerUniqueId();
+			if (ensure(LocalUserId.IsValid() && LocalUserId->IsValid()))
+			{
+				// Leave the party (restored in frontend)
+				LeavePersistentParty(*LocalUserId, InCompletionDelegate);
+			}
+		}
+		else
+		{
+			// Just block new joining until back in the frontend
+			PersistentParty->SetAcceptingMembers(false, EJoinPartyDenialReason::Busy);
+		}
+	}
+}
+
+void UParty::LeaveAndRestorePersistentParty()
+{
+	if (!bLeavingPersistentParty)
+	{
+		bLeavingPersistentParty = true;
+		UWorld* World = GetWorld();
+
+		FTimerDelegate LeaveAndRestorePersistentPartyNextTick;
+		LeaveAndRestorePersistentPartyNextTick.BindUObject(this, &ThisClass::LeaveAndRestorePersistentPartyInternal);
+		World->GetTimerManager().SetTimerForNextTick(LeaveAndRestorePersistentPartyNextTick);
+	}
+	else
+	{
+		UE_LOG(LogParty, Verbose, TEXT("Already leaving persistent party, ignoring"));
+	}
+}
+
+void UParty::LeaveAndRestorePersistentPartyInternal()
+{
+	UGameInstance* GameInstance = GetGameInstance();
+	check(GameInstance);
+
+	TSharedPtr<const FUniqueNetId> PrimaryUserId = GameInstance->GetPrimaryPlayerUniqueId();
+	check(PrimaryUserId.IsValid() && PrimaryUserId->IsValid());
+
+	UPartyDelegates::FOnLeaveUPartyComplete CompletionDelegate;
+	CompletionDelegate.BindUObject(this, &ThisClass::OnLeavePersistentPartyAndRestore);
+
+	// Unset this here, wouldn't be here if we were unable to call LeaveAndRestorePersistentParty, it will be reset inside
+	ensure(bLeavingPersistentParty);
+	bLeavingPersistentParty = false;
+	LeavePersistentParty(*PrimaryUserId, CompletionDelegate);
+}
+
+void UParty::OnLeavePersistentPartyAndRestore(const FUniqueNetId& LocalUserId, const ELeavePartyCompletionResult Result)
+{
+	UE_LOG(LogParty, Display, TEXT("OnLeavePersistentPartyAndRestore Result: %s"), ToString(Result));
+
+	RestorePersistentPartyState();
+}
+
+void UParty::AddPendingPartyJoin(const FUniqueNetId& LocalUserId, const FPartyDetails& PartyDetails, const UPartyDelegates::FOnJoinUPartyComplete& JoinCompleteDelegate)
+{
+	if (LocalUserId.IsValid() && PartyDetails.IsValid())
+	{
+		if (!HasPendingPartyJoin())
+		{
+			PendingPartyJoin = MakeShareable(new FPendingPartyJoin(LocalUserId.AsShared(), PartyDetails, JoinCompleteDelegate));
+		}
+	}
+}
+
+void UParty::ClearPendingPartyJoin()
+{
+	PendingPartyJoin.Reset();
+}
+
+bool UParty::ProcessPendingPartyJoin()
+{
+	if (HasPendingPartyJoin())
+	{
+		HandlePendingJoin();
+		return true;
+	}
+
+	return false;
+}
+
+UWorld* UParty::GetWorld() const
+{
+	UGameInstance* GameInstance = Cast<UGameInstance>(GetOuter());
+	if (GameInstance)
+	{
+		return GameInstance->GetWorld();
+	}
+	
+	return nullptr;
+}
+
+UGameInstance* UParty::GetGameInstance() const
+{
+	return Cast<UGameInstance>(GetOuter());
+}
+
+#undef LOCTEXT_NAMESPACE
