Copying //Tasks/UE4/Dev-Online-Plugins-v3 to //UE4/Dev-Main (Source: //Tasks/UE4/Dev-Online-Plugins-v3 @ 3047294)

diff --git a/Engine/Plugins/Online/OnlineFramework/Source/Party/Private/PartyGameState.cpp b/Engine/Plugins/Online/OnlineFramework/Source/Party/Private/PartyGameState.cpp
new file mode 100644
--- /dev/null
+++ b/Engine/Plugins/Online/OnlineFramework/Source/Party/Private/PartyGameState.cpp
@@ -0,0 +1,1440 @@
+// Copyright 1998-2016 Epic Games, Inc. All Rights Reserved.
+
+#include "PartyPrivatePCH.h"
+#include "PartyGameState.h"
+#include "PartyMemberState.h"
+#include "Party.h"
+#include "PartyBeaconClient.h"
+#include "Online.h"
+#include "OnlineSubsystemUtils.h"
+
+namespace PartyConsoleVariables
+{
+	// CVars
+	TAutoConsoleVariable<int32> CVarAcceptJoinsDuringLoad(
+		TEXT("Party.CVarAcceptJoinsDuringLoad"),
+		1,
+		TEXT("Enables joins while leader is trying to load into a game\n")
+		TEXT("1 Enables. 0 disables."),
+		ECVF_Default);
+}
+
+
+UPartyGameState::UPartyGameState(const FObjectInitializer& ObjectInitializer) :
+	Super(ObjectInitializer),
+	PartyStateRefDef(nullptr),
+	PartyStateRef(nullptr),
+	OwningUserId(nullptr),
+	bDebugAcceptingMembers(false),
+	OssParty(nullptr),
+	bPromotionLockoutState(false),
+	bStayWithPartyOnDisconnect(false),
+	ReservationBeaconClientClass(nullptr),
+	ReservationBeaconClient(nullptr)
+{
+	PartyMemberStateClass = UPartyMemberState::StaticClass();
+	ReservationBeaconClientClass = APartyBeaconClient::StaticClass();
+
+	if (!HasAnyFlags(RF_ClassDefaultObject))
+	{
+	}
+}
+
+void UPartyGameState::BeginDestroy()
+{
+	Super::BeginDestroy();
+	OnShutdown();
+}
+
+void UPartyGameState::OnShutdown()
+{
+	if (!HasAnyFlags(RF_ClassDefaultObject))
+	{
+		UnregisterFrontendDelegates();
+	}
+
+	OssParty.Reset();
+	PartyMembersState.Empty();
+
+	if (PartyStateRefDef)
+	{
+		if (PartyStateRefScratch)
+		{
+			PartyStateRefDef->DestroyStruct(PartyStateRefScratch);
+			FMemory::Free(PartyStateRefScratch);
+			PartyStateRefScratch = nullptr;
+		}
+		PartyStateRefDef = nullptr;
+	}
+
+	PartyStateRef = nullptr;
+}
+
+/*static*/ void UPartyGameState::AddReferencedObjects(UObject* InThis, FReferenceCollector& Collector)
+{
+	Super::AddReferencedObjects(InThis, Collector);
+
+	UPartyGameState* This = CastChecked<UPartyGameState>(InThis);
+	check(This);
+
+	TArray<UPartyMemberState*> PartyMembers;
+	This->PartyMembersState.GenerateValueArray(PartyMembers);
+	Collector.AddReferencedObjects(PartyMembers);
+}
+
+void UPartyGameState::RegisterFrontendDelegates()
+{
+	UWorld* World = GetWorld();
+
+	UnregisterFrontendDelegates();
+}
+
+void UPartyGameState::UnregisterFrontendDelegates()
+{
+}
+
+void UPartyGameState::ResetPartyState()
+{
+	if (PartyStateRef)
+	{
+		// Reset Party State
+		PartyStateRef->Reset();
+	}
+}
+
+void UPartyGameState::ResetPartySize()
+{
+	UParty* Party = GetPartyOuter();
+	SetPartyMaxSize(Party->GetDefaultPartyMaxSize());
+}
+
+void UPartyGameState::ResetLocalPlayerState()
+{
+	UWorld* World = GetWorld();
+	check(World);
+	AGameState* GameState = World->GetGameState();
+	check(GameState);
+
+	for (const APlayerState* PlayerState : GameState->PlayerArray)
+	{
+		if (PlayerState && PlayerState->UniqueId.IsValid())
+		{
+			UPartyMemberState** LocalPartyMemberStatePtr = PartyMembersState.Find(PlayerState->UniqueId);
+			UPartyMemberState* LocalPartyMemberState = LocalPartyMemberStatePtr ? *LocalPartyMemberStatePtr : nullptr;
+			if (LocalPartyMemberState)
+			{
+				LocalPartyMemberState->Reset();
+			}
+		}
+	}
+}
+
+void UPartyGameState::Init(const FUniqueNetId& LocalUserId, TSharedPtr<const FOnlineParty>& InParty)
+{
+	if (ensure(InParty.IsValid()))
+	{
+		if (ensure(LocalUserId.IsValid()))
+		{
+			OwningUserId.SetUniqueNetId(LocalUserId.AsShared());
+			OssParty = InParty;
+
+			CurrentConfig = *OssParty->Config;
+			bDebugAcceptingMembers = CurrentConfig.bIsAcceptingMembers;
+
+			// Last since it needs the party info/id set first
+			RegisterFrontendDelegates();
+		}
+		else
+		{
+			UE_LOG(LogParty, Warning, TEXT("Init: Invalid owning user id!"));
+		}
+	}
+	else
+	{
+		UE_LOG(LogParty, Warning, TEXT("Init: Invalid party!"));
+	}
+}
+
+void UPartyGameState::InitFromCreate(const FUniqueNetId& LocalUserId, TSharedPtr<const FOnlineParty>& InParty)
+{
+	if (ensure(InParty.IsValid()))
+	{
+		Init(LocalUserId, InParty);
+
+		// Setup initial data for the party
+		ResetPartyState();
+
+		FUniqueNetIdRepl MemberId(LocalUserId.AsShared());
+		UpdatePartyData(MemberId);
+
+		// Make sure we create the local player's entry before broadcasting the join
+		SendLocalPlayerPartyData();
+
+		// Broadcast join
+		UParty* Party = GetPartyOuter();
+		check(Party);
+		Party->OnPartyJoined().Broadcast(this);
+	}
+	else
+	{
+		UE_LOG(LogParty, Warning, TEXT("InitFromCreate: Invalid party!"));
+	}
+}
+
+void UPartyGameState::InitFromJoin(const FUniqueNetId& LocalUserId, TSharedPtr<const FOnlineParty>& InParty)
+{
+	if (ensure(InParty.IsValid()))
+	{
+		Init(LocalUserId, InParty);
+
+		// Make sure we create the local player's entry before broadcasting the join
+		SendLocalPlayerPartyData();
+
+		// Broadcast join
+		UParty* Party = GetPartyOuter();
+		check(Party);
+		Party->OnPartyJoined().Broadcast(this);
+	}
+	else
+	{
+		UE_LOG(LogParty, Warning, TEXT("InitFromJoin: Invalid party!"));
+	}
+}
+
+void UPartyGameState::PreClientTravel()
+{
+	if (!PartyConsoleVariables::CVarAcceptJoinsDuringLoad.GetValueOnGameThread())
+	{
+		// Possibly deal with pending approvals?
+		RejectAllPendingJoinRequests();
+	}
+	CleanupReservationBeacon();
+
+	UnregisterFrontendDelegates();
+}
+
+bool UPartyGameState::ResetForFrontend()
+{
+	UE_LOG(LogParty, Verbose, TEXT("Resetting parties for frontend"));
+
+	bool bSuccess = false;
+
+	CleanupReservationBeacon();
+
+	if (ensure(OssParty.IsValid()))
+	{
+		if (ensure(OwningUserId.IsValid()))
+		{
+			UWorld* World = GetWorld();
+			IOnlinePartyPtr PartyInt = Online::GetPartyInterface(World);
+			if (ensure(PartyInt.IsValid()))
+			{
+				TSharedPtr<const FOnlineParty> ExistingParty = PartyInt->GetParty(*OwningUserId, *OssParty->PartyId);
+				if (ExistingParty.IsValid())
+				{
+					bSuccess = true;
+
+					Init(*OwningUserId, ExistingParty);
+					bStayWithPartyOnDisconnect = false;
+
+					ResetPartyState();
+					ResetLocalPlayerState();
+
+					const bool bIsPartyLeader = IsLocalPartyLeader();
+
+					// Refresh local player data
+					SendLocalPlayerPartyData();
+
+					// Remove members we have, but lower level doesn't know about
+					for (const auto& PartyMember : PartyMembersState)
+					{
+						const FUniqueNetIdRepl& MemberId = PartyMember.Key;
+						if (MemberId.IsValid())
+						{
+							TSharedPtr<FOnlinePartyMember> CheckPartyMember = PartyInt->GetPartyMember(*OwningUserId, *OssParty->PartyId, *MemberId);
+							if (!CheckPartyMember.IsValid())
+							{
+								UE_LOG(LogParty, Verbose, TEXT("[%s] Player %s left during fixup"), *OssParty->PartyId->ToString(), *MemberId.ToString());
+								HandlePartyMemberLeft(*MemberId, EMemberExitedReason::Left);
+							}
+						}
+					}
+
+					// Add members we don't have, but lower level does
+					TArray<TSharedRef<FOnlinePartyMember>> PartyMembers;
+					PartyInt->GetPartyMembers(*OwningUserId, *OssParty->PartyId, PartyMembers);
+					for (const auto& PartyMember : PartyMembers)
+					{
+						TSharedRef<const FUniqueNetId> MemberId = PartyMember->GetUserId();
+
+						FUniqueNetIdRepl UniqueId(MemberId);
+						UPartyMemberState** CurrentPartyMemberDataPtr = PartyMembersState.Find(UniqueId);
+						UPartyMemberState* CurrentPartyMemberData = CurrentPartyMemberDataPtr ? *CurrentPartyMemberDataPtr : nullptr;
+						if (!CurrentPartyMemberData)
+						{
+							TSharedPtr<FOnlinePartyData> PartyMemberData = PartyInt->GetPartyMemberData(*OwningUserId, *OssParty->PartyId, *MemberId);
+							if (PartyMemberData.IsValid())
+							{
+								UE_LOG(LogParty, Verbose, TEXT("[%s] Player %s data received during fixup"), *OssParty->PartyId->ToString(), *UniqueId.ToString());
+								HandlePartyMemberDataReceived(*MemberId, PartyMemberData.ToSharedRef());
+							}
+						}
+					}
+
+					if (bIsPartyLeader)
+					{
+						UpdatePartyData(OwningUserId);
+						ResetPartySize();
+						UpdateAcceptingMembers();
+					}
+				}
+				else
+				{
+					UE_LOG(LogParty, Warning, TEXT("Party interface can't find party during reset!"));
+				}
+			}
+			else
+			{
+				UE_LOG(LogParty, Warning, TEXT("Invalid party interface during reset!"));
+			}
+		}
+		else
+		{
+			UE_LOG(LogParty, Warning, TEXT("Invalid owning user during reset!"));
+		}
+	}
+	else
+	{
+		UE_LOG(LogParty, Warning, TEXT("Invalid party info during reset!"));
+	}
+
+	if (!bSuccess)
+	{
+		ResetPartyState();
+		ResetLocalPlayerState();
+		UnregisterFrontendDelegates();
+	}
+
+	return bSuccess;
+}
+
+UPartyMemberState* UPartyGameState::CreateNewPartyMember(const FUniqueNetId& InMemberId)
+{
+	UPartyMemberState* NewPartyMemberState = nullptr;
+
+	if (ensure(InMemberId.IsValid()))
+	{
+		UWorld* World = GetWorld();
+		check(World);
+		IOnlinePartyPtr PartyInt = Online::GetPartyInterface(World);
+		if (ensure(PartyInt.IsValid()))
+		{
+			TSharedPtr<FOnlinePartyMember> PartyMember = PartyInt->GetPartyMember(*OwningUserId, *OssParty->PartyId, InMemberId);
+			if (PartyMember.IsValid())
+			{
+				NewPartyMemberState = NewObject<UPartyMemberState>(this, PartyMemberStateClass);
+				NewPartyMemberState->UniqueId.SetUniqueNetId(PartyMember->GetUserId());
+				NewPartyMemberState->DisplayName = FText::FromString(PartyMember->GetDisplayName());
+			}
+			else
+			{
+				UE_LOG(LogParty, Warning, TEXT("CreateNewPartyMember: Invalid party member %s"), *InMemberId.ToString());
+			}
+		}
+		else
+		{
+			UE_LOG(LogParty, Warning, TEXT("CreateNewPartyMember: No party interface."));
+		}
+	}
+	else
+	{
+		UE_LOG(LogParty, Warning, TEXT("CreateNewPartyMember: Invalid member id."));
+	}
+
+	return NewPartyMemberState;
+}
+
+void UPartyGameState::HandlePartyConfigChanged( const TSharedRef<FPartyConfiguration>& InPartyConfig )
+{
+	UE_LOG(LogParty, VeryVerbose, TEXT("[%s] HandlePartyConfigChanged"), OssParty.IsValid() ? *OssParty->PartyId->ToString() : TEXT("INVALID"));
+	if (ensure(OssParty.IsValid()))
+	{
+		CurrentConfig = *OssParty->Config;
+		bDebugAcceptingMembers = CurrentConfig.bIsAcceptingMembers;
+	}
+}
+
+void UPartyGameState::HandlePartyMemberJoined(const FUniqueNetId& InMemberId)
+{
+	UE_LOG(LogParty, VeryVerbose, TEXT("[%s] HandlePartyMemberJoined %s"), OssParty.IsValid() ? *OssParty->PartyId->ToString() : TEXT("INVALID"), *InMemberId.ToString());
+
+	TSharedRef<const FUniqueNetId> IdRef = InMemberId.AsShared();
+	FUniqueNetIdRepl MemberId(IdRef);
+	if (MemberId.IsValid())
+	{
+		UPartyMemberState** NewPartyMemberPtr = PartyMembersState.Find(MemberId);
+		UPartyMemberState* NewPartyMember = NewPartyMemberPtr ? *NewPartyMemberPtr : nullptr;
+		if (!NewPartyMember)
+		{
+			NewPartyMember = CreateNewPartyMember(InMemberId);
+			if (NewPartyMember)
+			{
+				PartyMembersState.Add(MemberId, NewPartyMember);
+			}
+		}
+
+		if (ensure(NewPartyMember))
+		{
+			if (!NewPartyMember->bHasAnnouncedJoin)
+			{
+				// Both local and remote players will announce joins
+				UParty* Party = GetPartyOuter();
+				check(Party);
+				Party->OnPartyMemberJoined().Broadcast(this, MemberId);
+				NewPartyMember->bHasAnnouncedJoin = true;
+			}
+		}
+
+		UpdateAcceptingMembers();
+	}
+}
+
+void UPartyGameState::HandlePartyMemberLeft(const FUniqueNetId& InMemberId, EMemberExitedReason Reason)
+{
+	UE_LOG(LogParty, VeryVerbose, TEXT("[%s] HandlePartyMemberLeft %s"), OssParty.IsValid() ? *OssParty->PartyId->ToString() : TEXT("INVALID"), *InMemberId.ToString());
+
+	if (ensure(InMemberId.IsValid()))
+	{
+		TSharedRef<const FUniqueNetId> IdRef = InMemberId.AsShared();
+		FUniqueNetIdRepl MemberId(IdRef);
+		
+		UParty* Party = GetPartyOuter();
+		if (ensure(Party))
+		{
+			Party->OnPartyMemberLeaving().Broadcast(this, MemberId, Reason);
+		}
+
+		PartyMembersState.Remove(MemberId);
+
+		if (Party)
+		{ 
+			Party->OnPartyMemberLeft().Broadcast(this, MemberId, Reason);
+		}
+		
+		// Update party join state, will cause a failure on leader promotion currently
+		// because we can't tell the difference between "expected leader" and "actually the new leader"
+		UpdateAcceptingMembers();
+	}
+}
+
+void UPartyGameState::HandlePartyMemberPromoted(const FUniqueNetId& InMemberId)
+{
+	ensure(OssParty.IsValid());
+	UE_LOG(LogParty, VeryVerbose, TEXT("[%s] HandlePartyMemberPromoted %s"), OssParty.IsValid() ? *OssParty->PartyId->ToString() : TEXT("INVALID"), *InMemberId.ToString());
+
+	if (ensure(InMemberId.IsValid()))
+	{
+		UParty* Party = GetPartyOuter();
+		if (ensure(Party))
+		{
+			TSharedRef<const FUniqueNetId> IdRef = InMemberId.AsShared();
+			FUniqueNetIdRepl MemberId(IdRef);
+			Party->OnPartyMemberPromoted().Broadcast(this, MemberId);
+		}
+	}
+
+	// Now that the leader is gone and a new leader established, make sure the accepting state is correct
+	UpdateAcceptingMembers();
+}
+
+void UPartyGameState::ComparePartyData(const FPartyState& OldPartyData, const FPartyState& NewPartyData)
+{
+	// Client passenger view delegates, leader won't get these because they are driving
+	if (!IsLocalPartyLeader())
+	{
+		if (OldPartyData.PartyType != NewPartyData.PartyType)
+		{
+			OnPartyTypeChanged().Broadcast(NewPartyData.PartyType);
+		}
+
+		if (OldPartyData.bLeaderFriendsOnly != NewPartyData.bLeaderFriendsOnly)
+		{
+			OnLeaderFriendsOnlyChanged().Broadcast(NewPartyData.bLeaderFriendsOnly);
+		}
+
+		if (OldPartyData.bLeaderInvitesOnly != NewPartyData.bLeaderInvitesOnly)
+		{
+			OnLeaderInvitesOnlyChanged().Broadcast(NewPartyData.bLeaderInvitesOnly);
+		}
+	}
+}
+
+void UPartyGameState::HandlePartyDataReceived(const TSharedRef<FOnlinePartyData>& InPartyData)
+{
+	UE_LOG(LogParty, VeryVerbose, TEXT("[%s] HandlePartyDataReceived"), OssParty.IsValid() ? *OssParty->PartyId->ToString() : TEXT("INVALID"));
+
+	UWorld* World = GetWorld();
+	check(World);
+	IOnlinePartyPtr PartyInt = Online::GetPartyInterface(World);
+	if (ensure(PartyInt.IsValid()))
+	{
+		check(PartyStateRefDef && PartyStateRef);
+		if (FVariantDataConverter::VariantMapToUStruct(InPartyData->GetKeyValAttrs(), PartyStateRefDef, PartyStateRefScratch, 0, CPF_Transient | CPF_RepSkip))
+		{
+			ComparePartyData(*PartyStateRef, *PartyStateRefScratch);
+
+			ensure(PartyStateRefDef->GetCppStructOps()->Copy(PartyStateRef, PartyStateRefScratch, 1));
+			OnPartyDataChanged().Broadcast(*PartyStateRef);
+		}
+		else
+		{
+			UE_LOG(LogParty, Warning, TEXT("Failed to serialize party data!"));
+		}
+	}
+}
+
+void UPartyGameState::HandlePartyMemberDataReceived(const FUniqueNetId& InMemberId, const TSharedRef<FOnlinePartyData>& InPartyMemberData)
+{
+	UE_LOG(LogParty, VeryVerbose, TEXT("[%s] HandlePartyMemberDataReceived %s"), OssParty.IsValid() ? *OssParty->PartyId->ToString() : TEXT("INVALID"), *InMemberId.ToString());
+
+	UWorld* World = GetWorld();
+	check(World);
+	IOnlinePartyPtr PartyInt = Online::GetPartyInterface(World);
+	if (ensure(PartyInt.IsValid()))
+	{
+		FUniqueNetIdRepl UniqueId(InMemberId.AsShared());
+		UPartyMemberState** CurrentPartyMemberPtr = PartyMembersState.Find(UniqueId);
+		UPartyMemberState* CurrentPartyMember = CurrentPartyMemberPtr ? *CurrentPartyMemberPtr : nullptr;
+		if (!CurrentPartyMember)
+		{
+			CurrentPartyMember = CreateNewPartyMember(InMemberId);
+			if (CurrentPartyMember)
+			{
+				PartyMembersState.Add(UniqueId, CurrentPartyMember);
+			}
+		}
+
+		if (ensure(CurrentPartyMember))
+		{
+			if (!CurrentPartyMember->bHasAnnouncedJoin)
+			{
+				// Both local and remote players will announce joins
+				UParty* Party = GetPartyOuter();
+				check(Party);
+				Party->OnPartyMemberJoined().Broadcast(this, UniqueId);
+				CurrentPartyMember->bHasAnnouncedJoin = true;
+			}
+
+			check(CurrentPartyMember->MemberStateRefDef && CurrentPartyMember->MemberStateRefScratch);
+
+			if (FVariantDataConverter::VariantMapToUStruct(InPartyMemberData->GetKeyValAttrs(), CurrentPartyMember->MemberStateRefDef, CurrentPartyMember->MemberStateRef, 0, CPF_Transient | CPF_RepSkip))
+			{
+				// Broadcast property changes
+				CurrentPartyMember->ComparePartyMemberData(*CurrentPartyMember->MemberStateRefScratch);
+				// Copy out the old data
+				ensure(CurrentPartyMember->MemberStateRefDef->GetCppStructOps()->Copy(CurrentPartyMember->MemberStateRefScratch, CurrentPartyMember->MemberStateRef, 1));
+				OnPartyMemberDataChanged().Broadcast(CurrentPartyMember->UniqueId, CurrentPartyMember);
+			}
+			else
+			{
+				UE_LOG(LogParty, Warning, TEXT("[%s] Failed to serialize party member data!"), *InMemberId.ToString());
+				ensure(CurrentPartyMember->MemberStateRefDef->GetCppStructOps()->Copy(CurrentPartyMember->MemberStateRef, CurrentPartyMember->MemberStateRefScratch, 1));
+			}
+		}
+	}
+}
+
+void UPartyGameState::HandleLeavingParty()
+{
+	// Stop listening for delegates as we leave
+	UnregisterFrontendDelegates();
+}
+
+void UPartyGameState::HandleRemovedFromParty(EMemberExitedReason Reason)
+{
+	// Trigger delegate first
+	UParty* Party = GetPartyOuter();
+	if (ensure(Party))
+	{
+		Party->OnPartyLeft().Broadcast(this, Reason);
+	}
+
+	// Cleanup
+	OnShutdown();
+}
+
+void UPartyGameState::HandleLockoutPromotionStateChange(bool bNewLockoutState)
+{
+	bPromotionLockoutState = bNewLockoutState;
+}
+
+EApprovalAction UPartyGameState::ProcessJoinRequest(const FUniqueNetId& RecipientId, const FUniqueNetId& SenderId, EJoinPartyDenialReason& DenialReason)
+{
+	if (IsInJoinableGameState())
+	{
+		return EApprovalAction::Approve;
+	}
+	else
+	{
+		DenialReason = EJoinPartyDenialReason::GameFull;
+		return EApprovalAction::Deny;
+	}
+}
+
+void UPartyGameState::HandlePartyJoinRequestReceived(const FUniqueNetId& RecipientId, const FUniqueNetId& SenderId)
+{
+#if 0 // auto approve requests for debugging error case
+	UWorld* World = GetWorld();
+	IOnlinePartyPtr PartyInt = Online::GetPartyInterface(World);
+	if (PartyInt.IsValid())
+	{
+		PartyInt->ApproveJoinRequest(RecipientId, *Party->PartyId, SenderId, true);
+	}
+	return;
+#endif
+
+	EApprovalAction ApprovalAction = EApprovalAction::Deny;
+	EJoinPartyDenialReason DenialReason = EJoinPartyDenialReason::Busy;
+
+	if (IsLocalPartyLeader())
+	{
+		int32 NumPartyMembers = GetPartySize();
+		int32 MaxPartyMembers = CurrentConfig.MaxMembers;
+		if (NumPartyMembers < MaxPartyMembers)
+		{
+			// Give the game a chance to accept or deny this player
+			ApprovalAction = ProcessJoinRequest(RecipientId, SenderId, DenialReason);
+		}
+		else
+		{
+			DenialReason = EJoinPartyDenialReason::PartyFull;
+		}
+	}
+	else
+	{
+		// Party leader has changed
+		DenialReason = EJoinPartyDenialReason::NotPartyLeader;
+	}
+
+	if (ApprovalAction == EApprovalAction::Enqueue ||
+		ApprovalAction == EApprovalAction::EnqueueAndStartBeacon)
+	{
+		// Enqueue for a more opportune time
+		UE_LOG(LogParty, Verbose, TEXT("[%s] Enqueuing approval request for %s"), OssParty.IsValid() ? *OssParty->PartyId->ToString() : TEXT("INVALID"), *SenderId.ToString());
+		
+		FPendingMemberApproval PendingApproval;
+		PendingApproval.RecipientId.SetUniqueNetId(RecipientId.AsShared());
+		PendingApproval.SenderId.SetUniqueNetId(SenderId.AsShared());
+		PendingApprovals.Enqueue(PendingApproval);
+
+		if (ReservationBeaconClient == nullptr &&
+			ApprovalAction == EApprovalAction::EnqueueAndStartBeacon)
+		{
+			ConnectToReservationBeacon();
+		}
+	}
+	else
+	{
+		bool bApproveRequest = (ApprovalAction == EApprovalAction::Approve);
+
+		// Respond now
+		UE_LOG(LogParty, Verbose, TEXT("[%s] Responding to approval request for %s with %s"), OssParty.IsValid() ? *OssParty->PartyId->ToString() : TEXT("INVALID"), *SenderId.ToString(), bApproveRequest ? TEXT("approved") : TEXT("denied"));
+		
+		UWorld* World = GetWorld();
+		IOnlinePartyPtr PartyInt = Online::GetPartyInterface(World);
+		if (PartyInt.IsValid() && OssParty.IsValid())
+		{
+			PartyInt->ApproveJoinRequest(RecipientId, *OssParty->PartyId, SenderId, bApproveRequest, (int32)(!bApproveRequest ? DenialReason : EJoinPartyDenialReason::NoReason));
+		}
+	}
+}
+
+FOnlinePartyTypeId UPartyGameState::GetPartyTypeId() const
+{
+	FOnlinePartyTypeId Result;
+
+	if (ensure(OssParty.IsValid()))
+	{
+		Result = OssParty->PartyTypeId;
+	}
+
+	return Result;
+}
+
+TSharedPtr<const FOnlinePartyId> UPartyGameState::GetPartyId() const
+{
+	if (ensure(OssParty.IsValid()))
+	{
+		return OssParty->PartyId;
+	}
+
+	return nullptr;
+}
+
+void UPartyGameState::SetPartyType(EPartyType InPartyType, bool bLeaderFriendsOnly, bool bLeaderInvitesOnly)
+{
+	if (IsLocalPartyLeader())
+	{
+		check(PartyStateRef);
+		if (PartyStateRef->PartyType != InPartyType ||
+			PartyStateRef->bLeaderFriendsOnly != bLeaderFriendsOnly ||
+			PartyStateRef->bLeaderInvitesOnly != bLeaderInvitesOnly)
+		{
+			bool bIsPrivate = (InPartyType == EPartyType::Private);
+
+			PartySystemPermissions::EPresencePermissions PresencePermissions;
+			if (bLeaderFriendsOnly)
+			{
+				if (bIsPrivate)
+				{
+					PresencePermissions = PartySystemPermissions::FriendsInviteOnly;
+				}
+				else
+				{
+					PresencePermissions = PartySystemPermissions::FriendsOnly;
+				}
+			}
+			else
+			{
+				if (bIsPrivate)
+				{
+					PresencePermissions = PartySystemPermissions::PublicInviteOnly;
+				}
+				else
+				{
+					PresencePermissions = PartySystemPermissions::Public;
+				}
+			}
+
+			CurrentConfig.PresencePermissions = PresencePermissions;
+			if (PartyStateRef->bInvitesDisabled)
+			{
+				CurrentConfig.InvitePermissions = PartySystemPermissions::EInvitePermissions::Noone;
+			}
+			else
+			{
+				CurrentConfig.InvitePermissions = bLeaderInvitesOnly ? PartySystemPermissions::EInvitePermissions::Leader : PartySystemPermissions::EInvitePermissions::Anyone;
+			}
+
+			UpdatePartyConfig(bIsPrivate);
+
+			EPartyType OldPartyType = PartyStateRef->PartyType;
+			bool bOldLeaderFriendsOnly = PartyStateRef->bLeaderFriendsOnly;
+			bool bOldLeaderInvitesOnly = PartyStateRef->bLeaderInvitesOnly;
+
+			// Replicate the party settings to other party members
+			PartyStateRef->PartyType = InPartyType;
+			PartyStateRef->bLeaderFriendsOnly = bLeaderFriendsOnly;
+			PartyStateRef->bLeaderInvitesOnly = bLeaderInvitesOnly;
+			UpdatePartyData(OwningUserId);
+
+			// Refresh accepting members, taking everything into account
+			UpdateAcceptingMembers();
+
+			// Notify the local player of the changes
+			if (OldPartyType != InPartyType)
+			{
+				OnPartyTypeChanged().Broadcast(InPartyType);
+			}
+
+			if (bOldLeaderFriendsOnly != bLeaderFriendsOnly)
+			{
+				OnLeaderFriendsOnlyChanged().Broadcast(bLeaderFriendsOnly);
+			}
+
+			if (bOldLeaderInvitesOnly != bLeaderInvitesOnly)
+			{
+				OnLeaderInvitesOnlyChanged().Broadcast(bLeaderInvitesOnly);
+			}
+		}
+	}
+	else
+	{
+		UE_LOG(LogParty, Warning, TEXT("Non party leader trying to set party permissions!"));
+	}
+}
+
+void UPartyGameState::SetInvitesDisabled(bool bInvitesDisabled)
+{
+	if (IsLocalPartyLeader())
+	{
+		check(PartyStateRef);
+		if (PartyStateRef->bInvitesDisabled != bInvitesDisabled)
+		{
+			if (bInvitesDisabled)
+			{
+				CurrentConfig.InvitePermissions = PartySystemPermissions::EInvitePermissions::Noone;
+			}
+			else
+			{
+				CurrentConfig.InvitePermissions = PartyStateRef->bLeaderInvitesOnly ? PartySystemPermissions::EInvitePermissions::Leader : PartySystemPermissions::EInvitePermissions::Anyone;
+			}
+
+			UpdatePartyConfig();
+
+			// Replicate the party settings to other party members
+			PartyStateRef->bInvitesDisabled = bInvitesDisabled;
+			UpdatePartyData(OwningUserId);
+
+			// Refresh accepting members, taking everything into account
+			UpdateAcceptingMembers();
+
+			OnInvitesDisabledChanged().Broadcast(bInvitesDisabled);
+		}
+	}
+	else
+	{
+		UE_LOG(LogParty, Verbose, TEXT("Non party leader trying to set invites disabled!"));
+	}
+}
+
+void UPartyGameState::StayWithPartyOnExit(bool bInStayWithParty)
+{
+	bStayWithPartyOnDisconnect = bInStayWithParty;
+}
+
+bool UPartyGameState::ShouldStayWithPartyOnExit() const
+{
+	return bStayWithPartyOnDisconnect;
+}
+
+int32 UPartyGameState::GetPartySize() const
+{
+	return PartyMembersState.Num();
+}
+
+void UPartyGameState::SetPartyMaxSize(int32 NewSize)
+{
+	if (ensure(OssParty.IsValid()))
+	{
+		if (CurrentConfig.MaxMembers != NewSize)
+		{
+			UParty* Party = GetPartyOuter();
+			CurrentConfig.MaxMembers = FMath::Clamp(NewSize, 1, Party->GetDefaultPartyMaxSize());
+			UpdatePartyConfig();
+		}
+	}
+	else
+	{
+		UE_LOG(LogParty, Warning, TEXT("Invalid party updating party size!"));
+	}
+}
+
+int32 UPartyGameState::GetPartyMaxSize() const
+{
+	if (ensure(OssParty.IsValid()))
+	{
+		return OssParty->Config->MaxMembers;
+	}
+	else
+	{
+		UE_LOG(LogParty, Warning, TEXT("Invalid party getting party size!"));
+	}
+
+	// Invalid party info
+	return INDEX_NONE;
+}
+
+bool UPartyGameState::IsPartyFull() const
+{
+	return GetPartySize() >= GetPartyMaxSize();
+}
+
+void UPartyGameState::UpdateAcceptingMembers()
+{
+	if (IsLocalPartyLeader())
+	{
+		EJoinPartyDenialReason DenialReason = EJoinPartyDenialReason::NoReason;
+		bool bCurrentlyAcceptingMembers = false;
+
+		// Look at game joinability (in game with permission or no game at all)
+		if (IsInJoinableGameState())
+		{
+			// Make sure the party isn't full
+			int32 NumPartyMembers = GetPartySize();
+			int32 MaxPartyMembers = CurrentConfig.MaxMembers;
+			if (NumPartyMembers < MaxPartyMembers)
+			{
+				// Look at party joinability
+				switch (PartyStateRef->PartyType)
+				{
+					case EPartyType::Public:
+					case EPartyType::FriendsOnly:
+						bCurrentlyAcceptingMembers = true;
+						break;
+					case EPartyType::Private:
+						// Party is private, invite required
+						DenialReason = EJoinPartyDenialReason::PartyPrivate;
+					default:
+						break;
+				}
+			}
+			else
+			{
+				// Party is full
+				DenialReason = EJoinPartyDenialReason::PartyFull;
+			}
+		}
+		else
+		{
+			DenialReason = EJoinPartyDenialReason::GameFull;
+		}
+
+		SetAcceptingMembers(bCurrentlyAcceptingMembers, DenialReason);
+	}
+	else
+	{
+		UE_LOG(LogParty, Warning, TEXT("Non party leader trying to update accepting members!"));
+	}
+}
+
+void UPartyGameState::SetAcceptingMembers(bool bIsAcceptingMembers, EJoinPartyDenialReason DenialReason)
+{
+	if (IsLocalPartyLeader())
+	{
+		ensure(OssParty.IsValid());
+
+		int32 NumPartyMembers = GetPartySize();
+		int32 MaxPartyMembers = CurrentConfig.MaxMembers;
+		bool bIsRoomInParty = (NumPartyMembers < MaxPartyMembers);
+
+		bool bCanAcceptMembers = bIsAcceptingMembers && bIsRoomInParty;
+
+		int32 NewDenialReason = (int32)(bCanAcceptMembers ? EJoinPartyDenialReason::NoReason : DenialReason);
+		if (CurrentConfig.bIsAcceptingMembers != bCanAcceptMembers || CurrentConfig.NotAcceptingMembersReason != NewDenialReason)
+		{
+			bDebugAcceptingMembers = bCanAcceptMembers;
+			CurrentConfig.bIsAcceptingMembers = bDebugAcceptingMembers;
+			CurrentConfig.NotAcceptingMembersReason = NewDenialReason;
+			UpdatePartyConfig();
+		}
+	}
+	else
+	{
+		UE_LOG(LogParty, Warning, TEXT("Non party leader trying to set accepting members!"));
+	}
+}
+
+bool UPartyGameState::IsInJoinableGameState() const
+{
+	bool bInGame = false;
+	bool bGameJoinable = false;
+
+	UWorld* World = GetWorld();
+	IOnlineSessionPtr SessionInt = Online::GetSessionInterface(World);
+	if (ensure(SessionInt.IsValid()))
+	{
+		bool bGamePublicJoinable = false;
+		bool bGameFriendJoinable = false;
+		bool bGameInviteOnly = false;
+		bool bGameAllowInvites = false;
+
+		FNamedOnlineSession* GameSession = SessionInt->GetNamedSession(GameSessionName);
+		if (GameSession != NULL &&
+			GameSession->GetJoinability(bGamePublicJoinable, bGameFriendJoinable, bGameInviteOnly, bGameAllowInvites))
+		{
+			bInGame = true;
+
+			TSharedPtr<FOnlineSessionInfo> UserSessionInfo = GameSession->SessionInfo;
+			if (UserSessionInfo.IsValid())
+			{
+				// User's game is joinable in some way if any of this is true
+				bGameJoinable = bGamePublicJoinable || bGameFriendJoinable || bGameInviteOnly;
+			}
+		}
+	}
+
+	return !bInGame || (bInGame && bGameJoinable);
+}
+
+bool UPartyGameState::CanInvite() const
+{
+	if (ensure(OssParty.IsValid()))
+	{
+		return OssParty->CanLocalUserInvite(*OwningUserId);
+	}
+	return false;
+}
+
+void UPartyGameState::UpdatePartyConfig(bool bResetAccessKey)
+{
+	if (IsLocalPartyLeader())
+	{
+		ensure(OssParty.IsValid());
+
+		UWorld* World = GetWorld();
+		check(World);
+		IOnlinePartyPtr PartyInt = Online::GetPartyInterface(World);
+		if (ensure(PartyInt.IsValid()))
+		{
+			FOnUpdatePartyComplete CompletionDelegate;
+			CompletionDelegate.BindUObject(this, &ThisClass::OnUpdatePartyConfigComplete);
+			if (!PartyInt->UpdateParty(*OwningUserId, *OssParty->PartyId, CurrentConfig, bResetAccessKey, CompletionDelegate))
+			{
+				UE_LOG(LogParty, Warning, TEXT("[%s] Failed to update party"), *OssParty->PartyId->ToString());
+			}
+		}
+		else
+		{
+			UE_LOG(LogParty, Warning, TEXT("[%s] Invalid party interface updating party size"), *OssParty->PartyId->ToString());
+		}
+	}
+	else
+	{
+		UE_LOG(LogParty, Warning, TEXT("Non party leader trying to update party config!"));
+	}
+}
+
+void UPartyGameState::OnUpdatePartyConfigComplete(const FUniqueNetId& LocalUserId, const FOnlinePartyId& InPartyId, const EUpdateConfigCompletionResult Result)
+{
+	const FString PartyIdDebugString = InPartyId.ToDebugString();
+	UE_LOG(LogParty, Verbose, TEXT("[%s] Party config updated %s"), *PartyIdDebugString, ToString(Result));
+
+	if (ensure(OssParty.IsValid()))
+	{
+		CurrentConfig = *OssParty->Config;
+		bDebugAcceptingMembers = CurrentConfig.bIsAcceptingMembers;
+	}
+}
+
+bool UPartyGameState::IsLocalPartyLeader()
+{
+	if (ensure(OwningUserId.IsValid()))
+	{
+		TSharedPtr<const FUniqueNetId> LeaderId = GetPartyLeader();
+		if (LeaderId.IsValid())
+		{
+			if (*OwningUserId == *LeaderId)
+			{
+				return true;
+			}
+		}
+		else
+		{
+			UE_LOG(LogParty, Warning, TEXT("Unable to determine party leader!"));
+		}
+	}
+	else
+	{
+		UE_LOG(LogParty, Warning, TEXT("Invalid owning user id trying to determine party leader!"));
+	}
+
+	return false;
+}
+
+TSharedPtr<const FUniqueNetId> UPartyGameState::GetPartyLeader() const
+{
+	TSharedPtr<const FUniqueNetId> ReturnValue = nullptr;
+	if (ensure(OssParty.IsValid()))
+	{
+		ReturnValue = OssParty->LeaderId;
+	}
+
+	return ReturnValue;
+}
+
+UPartyMemberState* UPartyGameState::GetPartyMember(const FUniqueNetIdRepl& InUniqueId) const
+{
+	if (InUniqueId.IsValid())
+	{
+		UPartyMemberState* const* PartyMemberPtr = PartyMembersState.Find(InUniqueId);
+		return PartyMemberPtr ? *PartyMemberPtr : nullptr;
+	}
+	
+	return nullptr;
+}
+
+void UPartyGameState::GetAllPartyMembers(TArray<UPartyMemberState*>& PartyMembers) const
+{
+	PartyMembersState.GenerateValueArray(PartyMembers);
+}
+
+FChatRoomId UPartyGameState::GetChatRoomID() const
+{
+	ensure(OssParty.IsValid());
+	return OssParty->RoomId;
+}
+
+void UPartyGameState::OnPartyMemberPromoted(const FUniqueNetId& LocalUserId, const FOnlinePartyId& InPartyId, const FUniqueNetId& InMemberId, const EPromoteMemberCompletionResult Result)
+{
+	FString PartyIdDebugString = InPartyId.ToDebugString();
+	FString MemberIdDebugString = InMemberId.ToDebugString();
+	UE_LOG(LogParty, Verbose, TEXT("[%s] Player %s promoted %s %s"), *PartyIdDebugString, *LocalUserId.ToString(), *MemberIdDebugString, ToString(Result));
+}
+
+void UPartyGameState::PromoteMember(const FUniqueNetIdRepl& NewPartyLeader)
+{
+	if (IsLocalPartyLeader())
+	{
+		ensure(OssParty.IsValid());
+		if (ensure(NewPartyLeader.IsValid()))
+		{
+			if (!bPromotionLockoutState)
+			{
+				UWorld* World = GetWorld();
+				check(World);
+				IOnlinePartyPtr PartyInt = Online::GetPartyInterface(World);
+				if (ensure(PartyInt.IsValid()))
+				{
+					// Do any internal updates
+					PrePromoteMember();
+
+					FOnPromotePartyMemberComplete CompletionDelegate;
+					CompletionDelegate.BindUObject(this, &ThisClass::OnPartyMemberPromoted);
+					PartyInt->PromoteMember(*OwningUserId, *OssParty->PartyId, *NewPartyLeader, CompletionDelegate);
+				}
+			}
+			else
+			{
+				UE_LOG(LogParty, Verbose, TEXT("[%s] Promote member feature locked out."), *OssParty->PartyId->ToString());
+			}
+		}
+		else
+		{
+			UE_LOG(LogParty, Warning, TEXT("Trying to promote invalid party member to leader!"));
+		}
+	}
+	else
+	{
+		UE_LOG(LogParty, Warning, TEXT("Non party leader trying to promote party leader!"));
+	}
+}
+
+void UPartyGameState::PrePromoteMember()
+{
+}
+
+void UPartyGameState::OnPartyMemberKicked(const FUniqueNetId& LocalUserId, const FOnlinePartyId& InPartyId, const FUniqueNetId& InMemberId, const EKickMemberCompletionResult Result)
+{
+	FString PartyIdDebugString = InPartyId.ToDebugString();
+	FString MemberIdDebugString = InMemberId.ToDebugString();
+	UE_LOG(LogParty, Verbose, TEXT("[%s] Player %s kicked %s %s"), *PartyIdDebugString, *LocalUserId.ToString(), *MemberIdDebugString, ToString(Result));
+}
+
+void UPartyGameState::KickMember(const FUniqueNetIdRepl& PartyMemberToKick)
+{
+	if (IsLocalPartyLeader())
+	{
+		ensure(OssParty.IsValid());
+		if (ensure(PartyMemberToKick.IsValid()))
+		{
+			UPartyMemberState** LocalPartyMemberStatePtr = PartyMembersState.Find(PartyMemberToKick);
+			if (LocalPartyMemberStatePtr && *LocalPartyMemberStatePtr)
+			{
+				UWorld* World = GetWorld();
+				check(World);
+				IOnlinePartyPtr PartyInt = Online::GetPartyInterface(World);
+				if (ensure(PartyInt.IsValid()))
+				{
+					FOnKickPartyMemberComplete CompletionDelegate;
+					CompletionDelegate.BindUObject(this, &ThisClass::OnPartyMemberKicked);
+					PartyInt->KickMember(*OwningUserId, *OssParty->PartyId, *PartyMemberToKick, CompletionDelegate);
+				}
+			}
+			else
+			{
+				UE_LOG(LogParty, Warning, TEXT("Trying to kick player that is not in your party!"));
+			}
+		}
+		else
+		{
+			UE_LOG(LogParty, Warning, TEXT("Trying to kick invalid party member!"));
+		}
+	}
+	else
+	{
+		UE_LOG(LogParty, Warning, TEXT("Non party leader trying to kick party member!"));
+	}
+}
+
+UPartyMemberState* UPartyGameState::InitPartyMemberStateFromLocalPlayer(const ULocalPlayer* LocalPlayer)
+{
+	UPartyMemberState* LocalPartyMemberState = nullptr;
+	TSharedPtr<const FUniqueNetId> UniqueNetId = LocalPlayer->GetPreferredUniqueNetId();
+	if (UniqueNetId.IsValid())
+	{
+		UPartyMemberState** LocalPartyMemberStatePtr = PartyMembersState.Find(UniqueNetId);
+		LocalPartyMemberState = LocalPartyMemberStatePtr ? *LocalPartyMemberStatePtr : nullptr;
+		if (!LocalPartyMemberState)
+		{
+			LocalPartyMemberState = CreateNewPartyMember(*UniqueNetId);
+			if (LocalPartyMemberState)
+			{
+				PartyMembersState.Add(UniqueNetId, LocalPartyMemberState);
+			}
+		}
+	}
+	return LocalPartyMemberState;
+}
+
+void UPartyGameState::SendLocalPlayerPartyData()
+{
+	UWorld* World = GetWorld();
+	check(World);
+
+	for (FLocalPlayerIterator It(GEngine, World); It; ++It)
+	{
+		ULocalPlayer* LP = Cast<ULocalPlayer>(*It);
+		if (LP)
+		{
+			const UPartyMemberState* MemberStatePtr = InitPartyMemberStateFromLocalPlayer(LP);
+			if (MemberStatePtr != nullptr)
+			{
+				UpdatePartyMemberState(OwningUserId, MemberStatePtr);
+			}
+		}
+	}
+}
+
+void UPartyGameState::UpdatePartyData(const FUniqueNetIdRepl& InLocalUserId)
+{
+	if (IsLocalPartyLeader())
+	{
+		ensure(OwningUserId == InLocalUserId);
+
+		TSharedPtr<const FOnlinePartyId> PartyId = GetPartyId();
+		if (PartyId.IsValid())
+		{
+			UWorld* World = GetWorld();
+			check(World);
+
+			IOnlinePartyPtr PartyInt = Online::GetPartyInterface(World);
+			if (ensure(PartyInt.IsValid()))
+			{
+				FOnlinePartyData PartyData;
+				ensure(PartyStateRefDef != nullptr && PartyStateRef != nullptr);
+				if (FVariantDataConverter::UStructToVariantMap(PartyStateRefDef, PartyStateRef, PartyData.GetKeyValAttrs(), 0, CPF_Transient | CPF_RepSkip))
+				{
+					PartyInt->UpdatePartyData(*OwningUserId, *PartyId, PartyData);
+				}
+				else
+				{
+					UE_LOG(LogParty, Warning, TEXT("UpdatePartyData: Failed to update party data!"));
+				}
+			}
+			else
+			{
+				UE_LOG(LogParty, Warning, TEXT("UpdatePartyData: Invalid party interface!"));
+			}
+		}
+		else
+		{
+			UE_LOG(LogParty, Warning, TEXT("UpdatePartyData: Invalid internal party!"));
+		}
+	}
+	else
+	{
+		UE_LOG(LogParty, Warning, TEXT("Non party leader trying to update party state!"));
+	}
+}
+
+void UPartyGameState::UpdatePartyMemberState(const FUniqueNetIdRepl& InLocalUserId, const UPartyMemberState* InPartyMemberState)
+{
+	if (ensure(InLocalUserId.IsValid()))
+	{
+		if (ensure(InPartyMemberState))
+		{
+			TSharedPtr<const FOnlinePartyId> PartyId = GetPartyId();
+			if (PartyId.IsValid())
+			{
+				UWorld* World = GetWorld();
+				check(World);
+
+				IOnlinePartyPtr PartyInt = Online::GetPartyInterface(World);
+				if (ensure(PartyInt.IsValid()))
+				{
+					FOnlinePartyData PartyMemberData;
+					ensure(InPartyMemberState->MemberStateRefDef != nullptr && InPartyMemberState->MemberStateRef != nullptr);
+					if (FVariantDataConverter::UStructToVariantMap(InPartyMemberState->MemberStateRefDef, InPartyMemberState->MemberStateRef, PartyMemberData.GetKeyValAttrs(), 0, CPF_Transient | CPF_RepSkip))
+					{
+						PartyInt->UpdatePartyMemberData(*InLocalUserId, *PartyId, PartyMemberData);
+					}
+					else
+					{
+						UE_LOG(LogParty, Warning, TEXT("UpdatePartyMemberState: Failed to update party member data!"));
+					}
+				}
+				else
+				{
+					UE_LOG(LogParty, Warning, TEXT("UpdatePartyMemberState: Invalid party interface!"));
+				}
+			}
+			else
+			{
+				UE_LOG(LogParty, Warning, TEXT("UpdatePartyMemberState: Invalid internal party!"));
+			}
+		}
+		else
+		{
+			UE_LOG(LogParty, Warning, TEXT("UpdatePartyMemberState: NULL party member state!"));
+		}
+	}
+	else
+	{
+		UE_LOG(LogParty, Warning, TEXT("UpdatePartyMemberState: Invalid local user!"));
+	}
+}
+
+void UPartyGameState::ConnectToReservationBeacon()
+{
+	if (IsLocalPartyLeader())
+	{
+		FPendingMemberApproval NextApproval;
+		if (PendingApprovals.Peek(NextApproval))
+		{
+			UWorld* World = GetWorld();
+			check(World);
+
+			// Reconnect to the reservation beacon to maintain our place in the game (just until actual joined, holds place for all party members)
+			ReservationBeaconClient = World->SpawnActor<APartyBeaconClient>(ReservationBeaconClientClass);
+			if (ReservationBeaconClient)
+			{
+				UE_LOG(LogParty, Verbose, TEXT("Created party reservation beacon %s."), *ReservationBeaconClient->GetName());
+
+				ReservationBeaconClient->OnHostConnectionFailure().BindUObject(this, &ThisClass::OnReservationBeaconUpdateConnectionFailure);
+				ReservationBeaconClient->OnReservationRequestComplete().BindUObject(this, &ThisClass::OnReservationBeaconUpdateResponseReceived);
+				ReservationBeaconClient->OnReservationCountUpdate().BindUObject(this, &ThisClass::OnReservationCountUpdate);
+
+				FPlayerReservation NewPlayerRes;
+				NewPlayerRes.UniqueId = NextApproval.SenderId;
+
+				TArray<FPlayerReservation> PlayersToAdd;
+				PlayersToAdd.Add(NewPlayerRes);
+
+				FUniqueNetIdRepl PartyLeader(GetPartyLeader());
+				ReservationBeaconClient->RequestReservationUpdate(CurrentSession, PartyLeader, PlayersToAdd);
+			}
+			else
+			{
+				OnReservationBeaconUpdateConnectionFailure();
+			}
+		}
+	}
+	else
+	{
+		UE_LOG(LogParty, Warning, TEXT("ConnectToReservationBeacon: Non party leader trying to connect to reservation beacon!"));
+	}
+}
+
+void UPartyGameState::RejectAllPendingJoinRequests()
+{
+	UWorld* World = GetWorld();
+	check(World);
+	IOnlinePartyPtr PartyInt = Online::GetPartyInterface(World);
+	TSharedPtr<const FOnlinePartyId> PartyId = GetPartyId();
+
+	bool bValidInterface = PartyInt.IsValid() && PartyId.IsValid();
+
+	bool bEmpty = false;
+	FPendingMemberApproval PendingApproval;
+
+	do
+	{
+		bEmpty = !PendingApprovals.Dequeue(PendingApproval);
+		if (!bEmpty && bValidInterface)
+		{
+			UE_LOG(LogParty, Verbose, TEXT("[%s] Responding to approval request for %s with denied"), *PartyId->ToString(), *PendingApproval.SenderId.ToString());
+			PartyInt->ApproveJoinRequest(*PendingApproval.RecipientId, *PartyId, *PendingApproval.SenderId, false, (int32)EJoinPartyDenialReason::Busy);
+		}
+	} while (!bEmpty);
+}
+
+void UPartyGameState::OnReservationBeaconUpdateConnectionFailure()
+{
+	UE_LOG(LogParty, Verbose, TEXT("Reservation update beacon failure %s."), ReservationBeaconClient ? *ReservationBeaconClient->GetName() : *FString());
+
+	// empty the queue, denying all requests
+	RejectAllPendingJoinRequests();
+	CleanupReservationBeacon();
+}
+
+void UPartyGameState::OnReservationBeaconUpdateResponseReceived(EPartyReservationResult::Type ReservationResponse)
+{
+	UE_LOG(LogParty, Verbose, TEXT("OnReservationBeaconUpdateResponseReceived %s"), EPartyReservationResult::ToString(ReservationResponse));
+
+	if (ReservationResponse == EPartyReservationResult::ReservationAccepted ||
+		ReservationResponse == EPartyReservationResult::ReservationDuplicate)
+	{
+		UWorld* World = GetWorld();
+
+		IOnlinePartyPtr PartyInt = Online::GetPartyInterface(World);
+		TSharedPtr<const FOnlinePartyId> PartyId = GetPartyId();
+
+		bool bValidInterface = PartyInt.IsValid() && PartyId.IsValid();
+
+		// There should be at least the one
+		FPendingMemberApproval PendingApproval;
+		if (ensure(PendingApprovals.Dequeue(PendingApproval)))
+		{
+			if (bValidInterface)
+			{
+				UE_LOG(LogParty, Verbose, TEXT("[%s] Responding to approval request for %s with approved"), *PartyId->ToString(), *PendingApproval.SenderId.ToString());
+				PartyInt->ApproveJoinRequest(*PendingApproval.RecipientId, *PartyId, *PendingApproval.SenderId, true);
+			}
+		}
+
+		// Check if there are any more while we are connected
+		if (PendingApprovals.Dequeue(PendingApproval))
+		{
+			if (ensure(ReservationBeaconClient))
+			{
+				FUniqueNetIdRepl PartyLeader(GetPartyLeader());
+
+				FPendingMemberApproval NextApproval;
+				PendingApprovals.Peek(NextApproval);
+
+				FPlayerReservation NewPlayerRes;
+				NewPlayerRes.UniqueId = NextApproval.SenderId;
+
+				TArray<FPlayerReservation> PlayersToAdd;
+				PlayersToAdd.Add(NewPlayerRes);
+
+				ReservationBeaconClient->RequestReservationUpdate(PartyLeader, PlayersToAdd);
+			}
+		}
+		else
+		{
+			CleanupReservationBeacon();
+		}
+	}
+	else
+	{
+		// empty the queue, denying all requests
+		RejectAllPendingJoinRequests();
+		CleanupReservationBeacon();
+	}
+}
+
+void UPartyGameState::OnReservationCountUpdate(int32 NumRemaining)
+{
+}
+
+void UPartyGameState::CleanupReservationBeacon()
+{
+	if (ReservationBeaconClient)
+	{
+		UE_LOG(LogParty, Verbose, TEXT("Party reservation beacon cleanup while in state %s, pending approvals: %s"), ToString(ReservationBeaconClient->GetConnectionState()), !PendingApprovals.IsEmpty() ? TEXT("true") : TEXT("false"));
+
+		ReservationBeaconClient->OnHostConnectionFailure().Unbind();
+		ReservationBeaconClient->OnReservationRequestComplete().Unbind();
+		ReservationBeaconClient->OnReservationCountUpdate().Unbind();
+		ReservationBeaconClient->DestroyBeacon();
+		ReservationBeaconClient = nullptr;
+	}
+}
+
+UWorld* UPartyGameState::GetWorld() const
+{
+	UParty* Party = GetPartyOuter();
+	if (Party)
+	{
+		return Party->GetWorld();
+	}
+	
+	return nullptr;
+}
+
+UParty* UPartyGameState::GetPartyOuter() const
+{
+	return GetTypedOuter<UParty>();
+}
diff --git a/Engine/Plugins/Online/OnlineFramework/Source/Party/Private/PartyMemberState.cpp b/Engine/Plugins/Online/OnlineFramework/Source/Party/Private/PartyMemberState.cpp
new file mode 100644
--- /dev/null
+++ b/Engine/Plugins/Online/OnlineFramework/Source/Party/Private/PartyMemberState.cpp
@@ -0,0 +1,84 @@
+// Copyright 1998-2016 Epic Games, Inc. All Rights Reserved.
+
+#include "PartyPrivatePCH.h"
+#include "PartyMemberState.h"
+#include "PartyGameState.h"
+#include "Party.h"
+
+UPartyMemberState::UPartyMemberState(const FObjectInitializer& ObjectInitializer) :
+	Super(ObjectInitializer),
+	MemberStateRefDef(nullptr),
+	MemberStateRef(nullptr),
+	MemberStateRefScratch(nullptr),
+	bHasAnnouncedJoin(false)	
+{
+	if (!HasAnyFlags(RF_ClassDefaultObject))
+	{
+	}
+}
+
+void UPartyMemberState::BeginDestroy()
+{
+	Super::BeginDestroy();
+
+	if (MemberStateRefDef)
+	{
+		if (MemberStateRefScratch)
+		{
+			MemberStateRefDef->DestroyStruct(MemberStateRefScratch);
+			FMemory::Free(MemberStateRefScratch);
+			MemberStateRefScratch = nullptr;
+		}
+		MemberStateRefDef = nullptr;
+	}
+
+	MemberStateRef = nullptr;
+}
+
+UPartyGameState* UPartyMemberState::GetParty() const
+{
+	return GetTypedOuter<UPartyGameState>();
+}
+
+bool UPartyMemberState::IsPartyLeader() const
+{
+	UPartyGameState* Party = GetParty();
+	check(Party);
+
+	TSharedPtr<const FUniqueNetId> PartyLeader = Party->GetPartyLeader();
+	if (PartyLeader.IsValid() && PartyLeader->IsValid())
+	{
+		return *PartyLeader == *UniqueId;
+	}
+
+	return false;
+}
+
+bool UPartyMemberState::IsLocalPlayer() const
+{
+	UPartyGameState* Party = GetParty();
+	check(Party);
+
+	return UniqueId == Party->GetOwningUserId();
+}
+
+void UPartyMemberState::UpdatePartyMemberState()
+{
+	UPartyGameState* Party = GetParty();
+	check(Party);
+	Party->UpdatePartyMemberState(UniqueId, this);
+}
+
+void UPartyMemberState::ComparePartyMemberData(const FPartyMemberRepState& OldPartyMemberState)
+{
+	UPartyGameState* Party = GetParty();
+	check(Party);
+}
+
+void UPartyMemberState::Reset()
+{
+	if (MemberStateRef)
+	{
+		MemberStateRef->Reset();
+	}
+}
diff --git a/Engine/Plugins/Online/OnlineFramework/Source/Party/Private/PartyPrivatePCH.h b/Engine/Plugins/Online/OnlineFramework/Source/Party/Private/PartyPrivatePCH.h
new file mode 100644
--- /dev/null
+++ b/Engine/Plugins/Online/OnlineFramework/Source/Party/Private/PartyPrivatePCH.h
@@ -0,0 +1,9 @@
+// Copyright 1998-2016 Epic Games, Inc. All Rights Reserved.
+
+#pragma once
+
+#include "Core.h"
+#include "Engine.h"
+
+// Module includes
+#include "PartyModule.h"
\ No newline at end of file
diff --git a/Engine/Plugins/Online/OnlineFramework/Source/Party/Public/Party.h b/Engine/Plugins/Online/OnlineFramework/Source/Party/Public/Party.h
new file mode 100644
--- /dev/null
+++ b/Engine/Plugins/Online/OnlineFramework/Source/Party/Public/Party.h
@@ -0,0 +1,528 @@
+// Copyright 1998-2016 Epic Games, Inc. All Rights Reserved.
+
+#pragma once
+
+#include "Online.h"
+#include "OnlineSubsystemTypes.h"
+#include "Party.generated.h"
+
+class UWorld;
+class FUniqueNetId;
+typedef FUniqueNetIdRepl FOnlinePartyIdRepl;
+class UPartyGameState;
+class UGameInstance;
+class IOnlinePartyJoinInfo;
+
+enum class EPartyType : uint8;
+
+/**
+ * Holds the basic information needed to join a party
+ */
+struct FPartyDetails
+{
+	FPartyDetails(const TSharedRef<const IOnlinePartyJoinInfo>& InPartyJoinInfo, bool bInAcceptInvite = true)
+		: PartyJoinInfo(InPartyJoinInfo), bAcceptInvite(bInAcceptInvite)
+	{}
+
+	bool IsValid() const
+	{
+		return PartyJoinInfo->IsValid();
+	}
+
+	const TSharedRef<const FOnlinePartyId>& GetPartyId() const
+	{
+		return PartyJoinInfo->GetPartyId();
+	}
+
+	const FOnlinePartyTypeId GetPartyTypeId() const
+	{
+		return PartyJoinInfo->GetPartyTypeId();
+	}
+
+	const TSharedRef<const FUniqueNetId>& GetLeaderId() const
+	{
+		return PartyJoinInfo->GetLeaderId();
+	}
+
+	const FString& GetAppId() const
+	{
+		return PartyJoinInfo->GetAppId();
+	}
+
+	FString ToString() const
+	{
+		return FString::Printf(
+			TEXT("PartyId: %s LeaderId: %s ResKey: %s Client: %s"), 
+			*GetPartyId()->ToDebugString(), 
+			*GetLeaderId()->ToDebugString(),
+			*GetAppId());
+	}
+
+	TSharedRef<const IOnlinePartyJoinInfo> PartyJoinInfo;
+	bool bAcceptInvite;
+};
+
+///////////////////////////////////////////////////////////////////
+// Completion delegates
+///////////////////////////////////////////////////////////////////
+namespace UPartyDelegates
+{
+	/**
+	 * Party creation async task completed callback
+	 *
+	 * @param LocalUserId - id of user that initiated the request
+	 * @param Result - result of the operation
+	 */
+	DECLARE_DELEGATE_TwoParams(FOnCreateUPartyComplete, const FUniqueNetId& /*LocalUserId*/, const ECreatePartyCompletionResult /*Result*/);
+	/**
+	 * Party join async task completed callback
+	 *
+	 * @param LocalUserId - id of user that initiated the request
+	 * @param Result - result of the operation
+	 * @param NotApprovedReason - client defined value describing why you were not approved
+	 */
+	DECLARE_DELEGATE_ThreeParams(FOnJoinUPartyComplete, const FUniqueNetId& /*LocalUserId*/, const EJoinPartyCompletionResult /*Result*/, const int32 /*NotApprovedReason*/);
+	/**
+	 * Party leave async task completed callback
+	 *
+	 * @param LocalUserId - id of user that initiated the request
+	 * @param Result - result of the operation
+	 */
+	DECLARE_DELEGATE_TwoParams(FOnLeaveUPartyComplete, const FUniqueNetId& /*LocalUserId*/, const ELeavePartyCompletionResult /*Result*/);
+}
+
+/**
+ * High level singleton for the management of parties, all parties are contained within
+ */
+UCLASS(config=Game, notplaceable)
+class PARTY_API UParty : public UObject
+{
+	GENERATED_UCLASS_BODY()
+
+private:
+
+	/**
+	 * Delegate fired when there is any data update on the party
+	 *
+	 * @param PartyUpdated party that updated
+	 */
+	DECLARE_MULTICAST_DELEGATE_OneParam(FOnPartyUpdate, UPartyGameState* /* PartyUpdated */);
+
+	/**
+	 * Delegate fired when there is any data update on a party member
+	 *
+	 * @param PartyState Party related to the data update
+	 * @param InMemberId UniqueId of the member data update
+	 */
+	DECLARE_MULTICAST_DELEGATE_TwoParams(FOnPartyMemberUpdate, UPartyGameState* /* PartyState */, const FUniqueNetIdRepl& /* InMemberId */);
+	
+	/**
+	 *	Delegate fired when the local player has left the party
+	 *	
+	 *	@param PartyState Party that the local player left
+	 *	@param Reason The reason why the player left the party
+	 */
+	DECLARE_MULTICAST_DELEGATE_TwoParams(FOnPartyLeft, UPartyGameState* /* PartyState */, EMemberExitedReason /* Reason */);
+
+	/**
+	*	Delegate fired when a party member has left the party
+	*
+	*	@param PartyState Party that the player left
+	*	@param InMmemberId UniqueId of the member that left
+	*	@param Reason The reason why the player left the party
+	*/
+	DECLARE_MULTICAST_DELEGATE_ThreeParams(FOnPartyMemberLeaving, UPartyGameState* /* PartyState */, const FUniqueNetIdRepl& /* InMemberId */, EMemberExitedReason /* Reason */);
+
+	/**
+	 *	Delegate fired when a party member has left the party
+	 *	
+	 *	@param PartyState Party that the player left
+	 *	@param InMmemberId UniqueId of the member that left
+	 *	@param Reason The reason why the player left the party
+	 */
+	DECLARE_MULTICAST_DELEGATE_ThreeParams(FOnPartyMemberLeft, UPartyGameState* /* PartyState */, const FUniqueNetIdRepl& /* InMemberId */, EMemberExitedReason /* Reason */);
+
+public:
+
+	/**
+	 * Initialization of the party management code, sets up listening for all party activity
+	 */
+	virtual void Init();
+
+	/**
+	 * Initialization specific to play in editor
+	 */
+	void InitPIE();
+
+	/**
+	 * Unregister delegates and clear out shared pointers to MCP objects
+	 */
+	void OnShutdown();
+
+	/**
+	 * Create a generic party
+	 *
+	 * @param InUserId user creating the party (should be primary player)
+	 * @param InPartyTypeId desired party type
+	 * @param InPartyConfig desired party configuration
+	 * @param InCompletionDelegate delegate called upon completion
+	 */
+	void CreateParty(const FUniqueNetId& InUserId, const FOnlinePartyTypeId InPartyTypeId, const FPartyConfiguration& InPartyConfig, const UPartyDelegates::FOnCreateUPartyComplete& InCompletionDelegate);
+
+	/**
+	 * Join a generic party
+	 *
+	 * @param InUserId user joining the party (should be primary player)
+	 * @param InPartyDetails credentials for the party to join
+	 * @param InCompletionDelegate delegate called upon completion
+	 */
+	void JoinParty(const FUniqueNetId& InUserId, const FPartyDetails& InPartyDetails, const UPartyDelegates::FOnJoinUPartyComplete& InCompletionDelegate);
+
+	/**
+	 * Leave a generic party
+	 *
+	 * @param InUserId user creating the party (should be primary player)
+	 * @param InPartyId id of the party to leave
+	 * @param InCompletionDelegate delegate called upon completion
+	 */
+	void LeaveParty(const FUniqueNetId& InUserId, const FOnlinePartyTypeId InPartyTypeId, const UPartyDelegates::FOnLeaveUPartyComplete& InCompletionDelegate);
+	
+	/**
+	 * Create a persistent party (party that will hold users for the duration of app)
+	 *
+	 * @param InUserId user creating the party (should be primary player)
+	 * @param InCompletionDelegate delegate called upon completion
+	 */
+	virtual void CreatePersistentParty(const FUniqueNetId& InUserId, const UPartyDelegates::FOnCreateUPartyComplete& InCompletionDelegate = UPartyDelegates::FOnCreateUPartyComplete());
+
+	/**
+	 * Join a persistent party (party that will hold users for the duration of app)
+	 *
+	 * @param InUserId user creating the party (should be primary player)
+	 * @param InPartyDetails credentials for the party to join
+	 * @param InCompletionDelegate delegate called upon completion
+	 */
+	virtual void JoinPersistentParty(const FUniqueNetId& InUserId, const FPartyDetails& InPartyDetails, const UPartyDelegates::FOnJoinUPartyComplete& InCompletionDelegate = UPartyDelegates::FOnJoinUPartyComplete());
+
+	/**
+	 * Leave a persistent party, a new one will need to be created for the given user afterward
+	 *
+	 * @param InUserId user creating the party (should be primary player)
+	 * @param InCompletionDelegate delegate called upon completion
+	 */
+	virtual void LeavePersistentParty(const FUniqueNetId& InUserId, const UPartyDelegates::FOnLeaveUPartyComplete& InCompletionDelegate = UPartyDelegates::FOnLeaveUPartyComplete());
+
+	/**
+	 * Reestablish all party state and information upon returning to the main menu
+	 * Makes sure all data is proper and notifies the UI to reconstruct the party
+	 */
+	void RestorePersistentPartyState();
+
+	/**
+	 * Do both LeavePersistentParty() and RestorePersistentPartyState() together
+	 * Typically for error cases, consider using the above functions first
+	 */
+	void LeaveAndRestorePersistentParty();
+
+	/**
+	 * Kick local players from persistent party
+	 *
+	 * @param InCompletionDelegate delegate called upon completion
+	 */
+	void KickFromPersistentParty(const UPartyDelegates::FOnLeaveUPartyComplete& InCompletionDelegate = UPartyDelegates::FOnLeaveUPartyComplete());
+
+	/**
+	 * Try to process any pending invites received while in the rest of the game
+	 *
+	 * @return true if a pending invite is available and a join will be attempted, false otherwise
+	 */
+	bool ProcessPendingPartyJoin();
+
+	/**
+	 * Take ownership of a pending invite and do what is necessary to get the game to a point where it can be used
+	 *
+	 * @param LocalUserId who the invite is to
+	 * @param PartyDetails party details used for the join
+	 * @param JoinCompleteDelegate delegate to call once the pending invite has been joined or in all failure cases
+	 */
+	void AddPendingPartyJoin(const FUniqueNetId& LocalUserId, const FPartyDetails& PartyDetails, const UPartyDelegates::FOnJoinUPartyComplete& JoinCompleteDelegate);
+
+	/** Clears the join data associated with a pending party join. */
+	void ClearPendingPartyJoin();
+
+	/**
+	 * Is any local player in the given party
+	 *
+	 * @param SessionId session to check
+	 *
+	 * @return true if any local player exists in the given party, false otherwise
+	 */
+	bool IsInParty(TSharedPtr<const FOnlinePartyId>& PartyId) const;
+
+	/** @return delegate fired when a party is joined */
+	FOnPartyUpdate& OnPartyJoined() { return PartyJoinedDelegates; }
+	/** @return delegate fired when the party is reset after returning to the frontend */
+	FOnPartyUpdate& OnPartyResetForFrontend() { return PartyResetForFrontendDelegate; }
+	/** @return delegate fired when a party is left */
+	FOnPartyLeft& OnPartyLeft() { return PartyLeftDelegates; }
+
+	/** @return delegate fired when a party member joins */
+	FOnPartyMemberUpdate& OnPartyMemberJoined() { return PartyMemberJoined; }
+	/** @return delegate fired when a party member is promoted to leader */
+	FOnPartyMemberUpdate& OnPartyMemberPromoted() { return PartyMemberPromoted; }
+	
+	/** @return delegate fired when a party member is about to leave */
+	FOnPartyMemberLeaving& OnPartyMemberLeaving() { return PartyMemberLeaving; }
+
+	/** @return delegate fired when a party member leaves */
+	FOnPartyMemberLeft& OnPartyMemberLeft() { return PartyMemberLeft; }
+
+	/**
+	 * Get the id for the current persistent party where the local players are members
+	 */
+	TSharedPtr<const FOnlinePartyId> GetPersistentPartyId() const { return PersistentPartyId; }
+	
+	/** @return the party state for a given party */
+	UPartyGameState* GetParty(const FOnlinePartyId& InPartyId) const;
+
+	/** @return the party state for a given party type */
+	UPartyGameState* GetParty(const FOnlinePartyTypeId InPartyTypeId) const;
+
+	/** @return the party state for the persistent party */
+	UPartyGameState* GetPersistentParty() const;
+
+	/** Notify the party system that travel is occurring */
+	void NotifyPreClientTravel(const FString& PendingURL, ETravelType TravelType, bool bIsSeamlessTravel);
+
+	/** @return true if the player has accepted an invite, but it hasn't been processed yet */
+	bool HasPendingPartyJoin();
+	/**
+	 * Party invites
+	 */
+
+protected:
+
+	/** Pending party join */
+	struct FPendingPartyJoin : public TSharedFromThis<FPendingPartyJoin>
+	{
+		FPendingPartyJoin(TSharedRef<const FUniqueNetId> InLocalUserId, const FPartyDetails& InPartyDetails, const UPartyDelegates::FOnJoinUPartyComplete& InDelegate)
+			: LocalUserId(InLocalUserId)
+			, PartyDetails(InPartyDetails)
+			, Delegate(InDelegate)
+		{
+			
+			ensure(LocalUserId->IsValid());
+		}
+
+		/** user that sent the invite */
+		TSharedRef<const FUniqueNetId> LocalUserId;
+		/** details about party to join */
+		FPartyDetails PartyDetails;
+
+		UPartyDelegates::FOnJoinUPartyComplete Delegate;
+	};
+
+	/** game invite info that is available to join via an accepted game invite */
+	TSharedPtr<FPendingPartyJoin> PendingPartyJoin;
+
+	/**
+	* Take ownership of a pending invite and do what is necessary to get the game to a point where it can be used
+	*
+	* @param LocalUserId who the invite is to
+	* @param PartyDetails party details used for the join
+	* @param JoinCompleteDelegate delegate to call once the pending invite has been joined or in all failure cases
+	*/
+	virtual void HandlePendingJoin() PURE_VIRTUAL(UParty::HandlePendingJoin, )
+
+	/**
+	 * Party configuration
+	 */
+
+public:
+
+	int32 GetDefaultPartyMaxSize() const { return DefaultMaxPartySize; }
+
+protected:
+
+	/** Mapping of party types to party classes, up to game to define relationship */
+	TMap<FOnlinePartyTypeId, TSubclassOf<UPartyGameState>> PartyClasses;
+
+	UPROPERTY(Config)
+	int32 DefaultMaxPartySize;
+
+	/**
+	 * Get the default settings for creating a persistent party
+	 *
+	 * @param PartyType party visibility setting
+	 * @param bLeaderFriendsOnly can only the leader send invites
+	 * @param bLeaderInvitesOnly can only the leader have friends join via presence
+	 * @param bAllowInvites are invites currently allowed
+	 */
+	virtual void GetDefaultPersistentPartySettings(EPartyType& PartyType, bool& bLeaderFriendsOnly, bool& bLeaderInvitesOnly, bool& bAllowInvites);
+
+	/**
+	 * Get the default configuration used for persistent party creation
+	 *
+	 * @param PartyConfig [out] desired party configuration
+	 */
+	void GetPersistentPartyConfiguration(FPartyConfiguration& PartyConfig);
+
+	/**
+	 * Party state management
+	 */
+	virtual void PartyMemberExitedInternal(const FUniqueNetId& InLocalUserId, const FOnlinePartyId& InPartyId, const FUniqueNetId& InMemberId, const EMemberExitedReason InReason);
+
+private:
+
+	/** Is leaving the persistent party already in flight */
+	UPROPERTY()
+	bool bLeavingPersistentParty;
+	/** Array of leave persistent party delegates gathered while already leaving a persistent party */
+	TArray<UPartyDelegates::FOnLeaveUPartyComplete> LeavePartyCompleteDelegates;
+	/** Id of the current persistent party, only one and it should always be valid except while transferring between parties */
+	TSharedPtr<const FOnlinePartyId> PersistentPartyId;
+	/** Mapping of all known joined parties */
+	TMap<FOnlinePartyTypeId, UPartyGameState*> JoinedParties;
+
+	/** Delegate when parties are joined */
+	FOnPartyUpdate PartyJoinedDelegates;
+	/** Delegate when the party is reset after returning to the home base */
+	FOnPartyUpdate PartyResetForFrontendDelegate;
+	/** Delegate when parties are removed */
+	FOnPartyLeft PartyLeftDelegates;
+
+	/** Delegate when party members are added */
+	FOnPartyMemberUpdate PartyMemberJoined;
+	/** Delegate when party members are promoted to leader */
+	FOnPartyMemberUpdate PartyMemberPromoted;
+	/** Delegate when party members are removed */
+	FOnPartyMemberLeaving PartyMemberLeaving;
+	/** Delegate when party members are removed */
+	FOnPartyMemberLeft PartyMemberLeft;
+
+	/**
+	 * Notification that a world has been loaded (need to get a UWorld to register delegates on OSS)
+	 */
+	void OnPostLoadMap();
+
+	/**
+	 * Register with the identity interface to monitor logout for party cleanup
+	 */
+	void RegisterIdentityDelegates();
+
+	/**
+	 * Unregister with the identity interface
+	 */
+	void UnregisterIdentityDelegates();
+
+	/**
+	 * Register with the party interface to monitor party events and dispatch to specific party state
+	 */
+	void RegisterPartyDelegates();
+
+	/**
+	 * Unregister with the party interface 
+	 */
+	void UnregisterPartyDelegates();
+
+	/**
+	 * Update the persistent party leader locally (possibly telling the server)
+	 *
+	 * @param NewPartyLeader id of new party leader
+	 */
+	virtual void UpdatePersistentPartyLeader(const FUniqueNetIdRepl& NewPartyLeader);
+
+	/**
+	 * Common calls for managing parties
+	 */
+	void CreatePartyInternal(const FUniqueNetId& InUserId, const FOnlinePartyTypeId InPartyTypeId, const FPartyConfiguration& InPartyConfig, const UPartyDelegates::FOnCreateUPartyComplete& InCompletionDelegate);
+	void JoinPartyInternal  (const FUniqueNetId& InUserId, const FPartyDetails& InPartyDetails, const UPartyDelegates::FOnJoinUPartyComplete& InCompletionDelegate);
+	void LeavePartyInternal(const FUniqueNetId& InUserId, const FOnlinePartyTypeId InPartyTypeId, const UPartyDelegates::FOnLeaveUPartyComplete& InCompletionDelegate);
+
+	/**
+	 * Delegates for managing parties
+	 */
+	void OnCreatePartyComplete(const FUniqueNetId& LocalUserId, const ECreatePartyCompletionResult Result, const FOnlinePartyTypeId InPartyTypeId, UPartyDelegates::FOnCreateUPartyComplete CompletionDelegate);
+	void OnJoinPartyComplete(const FUniqueNetId& LocalUserId, const EJoinPartyCompletionResult Result, int32 DeniedResultCode, const FOnlinePartyTypeId InPartyTypeId, UPartyDelegates::FOnJoinUPartyComplete CompletionDelegate);
+	void OnLeavePartyComplete(const FUniqueNetId& LocalUserId, const ELeavePartyCompletionResult Result, const FOnlinePartyTypeId InPartyTypeId, UPartyDelegates::FOnLeaveUPartyComplete CompletionDelegate);
+	void HandleJoinPersistentPartyFailure();
+
+	/**
+	 * Delegates for managing regular parties
+	 */
+	void OnCreatePartyInternalComplete(const FUniqueNetId& LocalUserId, const TSharedPtr<const FOnlinePartyId>& InPartyId, const ECreatePartyCompletionResult Result, const FOnlinePartyTypeId InPartyTypeId, UPartyDelegates::FOnCreateUPartyComplete CompletionDelegate);
+	void OnJoinPartyInternalComplete(const FUniqueNetId& LocalUserId, const FOnlinePartyId& InPartyId, const EJoinPartyCompletionResult Result, int32 DeniedResultCode, const FOnlinePartyTypeId InPartyTypeId, UPartyDelegates::FOnJoinUPartyComplete CompletionDelegate);
+	void OnLeavePartyInternalComplete(const FUniqueNetId& LocalUserId, const FOnlinePartyId& InPartyId, const ELeavePartyCompletionResult Result, const FOnlinePartyTypeId InPartyTypeId, UPartyDelegates::FOnLeaveUPartyComplete CompletionDelegate);
+
+	/**
+	 * Delegates for managing persistent parties
+	 */
+	void OnCreatePersistentPartyComplete(const FUniqueNetId& LocalUserId, const ECreatePartyCompletionResult Result, UPartyDelegates::FOnCreateUPartyComplete CompletionDelegate);
+	void OnJoinPersistentPartyComplete(const FUniqueNetId& LocalUserId, const EJoinPartyCompletionResult Result, int32 DeniedResultCode, UPartyDelegates::FOnJoinUPartyComplete CompletionDelegate);
+	void OnLeavePersistentPartyComplete(const FUniqueNetId& LocalUserId, const ELeavePartyCompletionResult Result, UPartyDelegates::FOnLeaveUPartyComplete CompletionDelegate);
+
+	/**
+	 * Party state management delegates 
+	 */
+	void PartyConfigChangedInternal(const FUniqueNetId& InLocalUserId, const FOnlinePartyId& InPartyId, const TSharedRef<FPartyConfiguration>& InPartyConfig);
+	void PartyMemberJoinedInternal(const FUniqueNetId& InLocalUserId, const FOnlinePartyId& InPartyId, const FUniqueNetId& InMemberId);
+	void PartyDataReceivedInternal(const FUniqueNetId& InLocalUserId, const FOnlinePartyId& InPartyId, const TSharedRef<FOnlinePartyData>& InPartyData);
+	void PartyMemberDataReceivedInternal(const FUniqueNetId& InLocalUserId, const FOnlinePartyId& InPartyId, const FUniqueNetId& InMemberId, const TSharedRef<FOnlinePartyData>& InPartyMemberData);
+	void PartyJoinRequestReceivedInternal(const FUniqueNetId& InLocalUserId, const FOnlinePartyId& InPartyId, const FUniqueNetId& SenderId);
+	void PartyMemberChangedInternal(const FUniqueNetId& InLocalUserId, const FOnlinePartyId& InPartyId, const FUniqueNetId& InMemberId, const EMemberChangedReason InReason);
+	void PartyPromotionLockoutStateChangedInternal(const FUniqueNetId& LocalUserId, const FOnlinePartyId& InPartyId, const bool bLockoutState);
+
+	void PartyExitedInternal(const FUniqueNetId& LocalUserId, const FOnlinePartyId& InPartyId);
+	void OnPersistentPartyExitedInternalCompleted(const FUniqueNetId& LocalUserId, const ECreatePartyCompletionResult Result);
+
+	/**
+	 *
+	 */
+	void OnCreatePesistentPartyCompletedCommon(const FUniqueNetId& LocalUserId);
+
+	/**
+	 * Identity interface status change delegates
+	 */
+	void OnLogoutComplete(int32 LocalUserNum, bool bWasSuccessful);
+	void OnLoginStatusChanged(int32 LocalUserNum, ELoginStatus::Type OldStatus, ELoginStatus::Type NewStatus, const FUniqueNetId& NewId);
+
+	/**
+	 * Error handling to restore persistent party after leaving a different one
+	 */
+	void LeaveAndRestorePersistentPartyInternal();
+	void OnLeavePersistentPartyAndRestore(const FUniqueNetId& LocalUserId, const ELeavePartyCompletionResult Result);
+
+	/**
+	 * Delegate handles
+	 */
+	FDelegateHandle PartyConfigChangedDelegateHandle;
+	FDelegateHandle PartyMemberJoinedDelegateHandle;
+	FDelegateHandle PartyPromotionLockoutChangedDelegateHandle;
+	FDelegateHandle PartyDataReceivedDelegateHandle;
+	FDelegateHandle PartyMemberDataReceivedDelegateHandle;
+	FDelegateHandle PartyJoinRequestReceivedDelegateHandle;
+	FDelegateHandle PartyMemberChangedDelegateHandle;
+	FDelegateHandle PartyMemberExitedDelegateHandle;
+	FDelegateHandle PartyExitedDelegateHandle;
+
+	FDelegateHandle LogoutStatusChangedDelegateHandle[MAX_LOCAL_PLAYERS];
+	FDelegateHandle LogoutCompleteDelegateHandle[MAX_LOCAL_PLAYERS];
+
+	/**
+	 * Helpers
+	 */
+
+public:
+
+	// UObject interface begin
+	virtual UWorld* GetWorld() const override;
+	static void AddReferencedObjects(UObject* InThis, FReferenceCollector& Collector);
+	// UObject interface end
+
+protected:
+
+	/** @return the game instance */
+	UGameInstance* GetGameInstance() const;
+
+};
+
diff --git a/Engine/Plugins/Online/OnlineFramework/Source/Party/Public/PartyMemberState.h b/Engine/Plugins/Online/OnlineFramework/Source/Party/Public/PartyMemberState.h
new file mode 100644
--- /dev/null
+++ b/Engine/Plugins/Online/OnlineFramework/Source/Party/Public/PartyMemberState.h
@@ -0,0 +1,98 @@
+// Copyright 1998-2016 Epic Games, Inc. All Rights Reserved.
+
+#pragma once
+
+#include "PartyMemberState.generated.h"
+
+class UPartyGameState;
+
+/**
+ * Simple struct for replication and copying of party member data on updates
+ */
+USTRUCT()
+struct FPartyMemberRepState
+{
+	GENERATED_USTRUCT_BODY()
+
+	/** Reset the variables of this party member state */
+	virtual void Reset()
+	{}
+};
+
+/**
+ * Main representation of a party member
+ */
+UCLASS(config = Game, notplaceable)
+class PARTY_API UPartyMemberState : public UObject
+{
+	GENERATED_UCLASS_BODY()
+
+	// Begin UObject Interface
+	virtual void BeginDestroy() override;
+	// End UObject Interface
+
+	/** Unique id of this party member */
+	UPROPERTY(Transient)
+	FUniqueNetIdRepl UniqueId;
+
+	/** Display name of this party member */
+	UPROPERTY(BlueprintReadOnly, Transient, Category = PartyMemberState)
+	FText DisplayName;
+
+	/** @return the party this member is associated with */
+	UPartyGameState* GetParty() const;
+
+	/** @return True if this party member is the party leader */
+	UFUNCTION(BlueprintCallable, Category = PartyMemberState)
+	bool IsPartyLeader() const;
+
+	/** @return True if this party member state corresponds to the local player */
+	UFUNCTION(BlueprintCallable, Category = PartyMemberState)
+	bool IsLocalPlayer() const;
+
+protected:
+
+	/** Reflection data for child USTRUCT */
+	UPROPERTY()
+	const UScriptStruct* MemberStateRefDef;
+
+	/** Pointer to child USTRUCT that holds the current state of party member (set via InitPartyMemberState) */
+	FPartyMemberRepState* MemberStateRef;
+
+	template<typename T>
+	void InitPartyMemberState(T* InMemberState)
+	{
+		MemberStateRef = InMemberState;
+		if (MemberStateRefDef == nullptr)
+		{
+			MemberStateRefDef = T::StaticStruct();
+		}
+
+		MemberStateRefScratch = (FPartyMemberRepState*)FMemory::Malloc(MemberStateRefDef->GetCppStructOps()->GetSize());
+		MemberStateRefDef->InitializeStruct(MemberStateRefScratch);
+	}
+
+	void UpdatePartyMemberState();
+
+	/**
+	 * Compare current data to old data, triggering delegates
+	 *
+	 * @param OldPartyMemberState old view of data to compare against
+	 */
+	virtual void ComparePartyMemberData(const FPartyMemberRepState& OldPartyMemberState);
+
+	/** Reset to default state */
+	virtual void Reset();
+
+private:
+	
+
+	/** Scratch copy of child USTRUCT for handling replication comparisons */
+	FPartyMemberRepState* MemberStateRefScratch;
+
+	/** Have we announced this player joining the game locally */
+	UPROPERTY(Transient)
+	bool bHasAnnouncedJoin;
+
+	friend UPartyGameState;
+};
diff --git a/Engine/Plugins/Online/OnlineFramework/Source/Qos/Private/QosEvaluator.cpp b/Engine/Plugins/Online/OnlineFramework/Source/Qos/Private/QosEvaluator.cpp
new file mode 100644
--- /dev/null
+++ b/Engine/Plugins/Online/OnlineFramework/Source/Qos/Private/QosEvaluator.cpp
@@ -0,0 +1,644 @@
+// Copyright 1998-2016 Epic Games, Inc. All Rights Reserved.
+
+#include "QosPrivatePCH.h"
+#include "QosEvaluator.h"
+#include "QosBeaconClient.h"
+#include "QosStats.h"
+#include "OnlineSubsystemUtils.h"
+#include "Icmp.h"
+
+FOnlineSessionSearchQos::FOnlineSessionSearchQos()
+{
+	bIsLanQuery = false;
+	MaxSearchResults = 100;
+
+	FString GameModeStr(GAMEMODE_QOS);
+	QuerySettings.Set(SETTING_GAMEMODE, GameModeStr, EOnlineComparisonOp::Equals);
+	QuerySettings.Set(SETTING_QOS, 1, EOnlineComparisonOp::Equals);
+}
+
+TSharedPtr<FOnlineSessionSettings> FOnlineSessionSearchQos::GetDefaultSessionSettings() const
+{
+	return MakeShareable(new FOnlineSessionSettingsQos());
+}
+
+void FOnlineSessionSearchQos::SortSearchResults()
+{
+	TMap<FString, int32> RegionCounts;
+
+	static const int32 MaxPerRegion = 5;
+
+	UE_LOG(LogQos, Verbose, TEXT("Sorting QoS results"));
+	for (int32 SearchResultIdx = 0; SearchResultIdx < SearchResults.Num();)
+	{
+		FString QosRegion;
+		FOnlineSessionSearchResult& SearchResult = SearchResults[SearchResultIdx];
+
+		if (!SearchResult.Session.SessionSettings.Get(SETTING_REGION, QosRegion) || QosRegion.IsEmpty())
+		{
+			UE_LOG(LogQos, Verbose, TEXT("Removed Qos search result, invalid region."));
+			SearchResults.RemoveAtSwap(SearchResultIdx);
+			continue;
+		}
+
+		int32& ResultCount = RegionCounts.FindOrAdd(QosRegion);
+		ResultCount++;
+		if (ResultCount > MaxPerRegion)
+		{
+			SearchResults.RemoveAtSwap(SearchResultIdx);
+			continue;
+		}
+
+		SearchResultIdx++;
+	}
+
+	for (auto& It : RegionCounts)
+	{
+		UE_LOG(LogQos, Verbose, TEXT("Region: %s Count: %d"), *It.Key, It.Value);
+	}
+}
+
+UQosEvaluator::UQosEvaluator(const FObjectInitializer& ObjectInitializer)
+	: Super(ObjectInitializer)
+	, ParentWorld(nullptr)
+	, QosSearchQuery(nullptr)
+	, QosBeaconClient(nullptr)
+	, bInProgress(false)
+	, bCancelOperation(false)
+	, AnalyticsProvider(nullptr)
+	, QosStats(nullptr)
+{
+}
+
+void UQosEvaluator::SetWorld(UWorld* InWorld)
+{
+	ParentWorld = InWorld;
+}
+
+void UQosEvaluator::SetAnalyticsProvider(TSharedPtr<IAnalyticsProvider> InAnalyticsProvider)
+{
+	AnalyticsProvider = InAnalyticsProvider;
+}
+
+void UQosEvaluator::Cancel()
+{
+	bCancelOperation = true;
+}
+
+void UQosEvaluator::FindDatacenters(const FQosParams& InParams, TArray<FQosDatacenterInfo>& InDatacenters, const FOnQosSearchComplete& InCompletionDelegate)
+{
+	if (!bInProgress)
+	{
+		bInProgress = true;
+		ControllerId = InParams.ControllerId;
+
+		double CurTimestamp = FPlatformTime::Seconds();
+
+		StartAnalytics();
+
+		Datacenters.Empty(InDatacenters.Num());
+		for (const FQosDatacenterInfo& Datacenter : InDatacenters)
+		{
+			if (Datacenter.IsUsable())
+			{
+				new (Datacenters) FQosRegionInfo(Datacenter);
+			}
+		}
+
+		if (!InParams.bUseOldQosServers)
+		{
+			// Ping list of known servers defined by config
+			PingRegionServers(InParams, InCompletionDelegate);
+		}
+		else
+		{
+			// Discover servers automatically
+			if (!FindQosServersByRegion(GetNextRegionId(), InCompletionDelegate))
+			{
+				// Failed to start
+				TWeakObjectPtr<UQosEvaluator> WeakThisCap(this);
+				GetWorldTimerManager().SetTimerForNextTick(FTimerDelegate::CreateLambda([InCompletionDelegate, WeakThisCap]()
+				{
+					if (WeakThisCap.IsValid())
+					{
+						auto StrongThis = WeakThisCap.Get();
+						StrongThis->FinalizeDatacenterResult(InCompletionDelegate, EQosCompletionResult::Failure, TArray<FQosRegionInfo>());
+					}
+				}));
+			}
+		}
+	}
+	else
+	{
+		UE_LOG(LogQos, Log, TEXT("Qos evaluation already in progress, ignoring"));
+		// Just trigger delegate now (Finalize resets state vars)
+		GetWorldTimerManager().SetTimerForNextTick(FTimerDelegate::CreateLambda([InCompletionDelegate]() {
+			InCompletionDelegate.ExecuteIfBound(EQosCompletionResult::Failure, TArray<FQosRegionInfo>());
+		}));
+	}
+}
+
+void UQosEvaluator::FinalizeDatacenterResult(const FOnQosSearchComplete& InCompletionDelegate, EQosCompletionResult CompletionResult, const TArray<FQosRegionInfo>& RegionInfo)
+{
+	UE_LOG(LogQos, Log, TEXT("Datacenter evaluation complete. Result: %s "), ToString(CompletionResult));
+
+	EndAnalytics(CompletionResult);
+
+	// Broadcast this data next frame
+	TWeakObjectPtr<UQosEvaluator> WeakThisCap(this);
+	GetWorldTimerManager().SetTimerForNextTick(FTimerDelegate::CreateLambda([InCompletionDelegate, CompletionResult, RegionInfo, WeakThisCap]()
+	{
+		InCompletionDelegate.ExecuteIfBound(CompletionResult, RegionInfo);
+		if (WeakThisCap.IsValid())
+		{
+			auto StrongThis = WeakThisCap.Get();
+			StrongThis->ResetSearchVars();
+		}
+	}));
+}
+
+void UQosEvaluator::PingRegionServers(const FQosParams& InParams, const FOnQosSearchComplete& InCompletionDelegate)
+{
+	// Failsafe for bad configuration
+	bool bDidNothing = true;
+
+	const int32 NumTestsPerRegion = InParams.NumTestsPerRegion;
+
+	TWeakObjectPtr<UQosEvaluator> WeakThisCap(this);
+ 	for (FQosRegionInfo& Region : Datacenters)
+	{
+		if (Region.Region.bEnabled)
+		{
+			const FString& RegionId = Region.Region.RegionId;
+			const int32 NumServers = Region.Region.Servers.Num();
+			int32 ServerIdx = FMath::RandHelper(NumServers);
+			Region.Result = NumServers > 0 ? EQosCompletionResult::Invalid : EQosCompletionResult::Success;
+			if (NumServers > 0)
+			{
+				for (int32 PingIdx = 0; PingIdx < NumTestsPerRegion; PingIdx++)
+				{
+					const FQosPingServerInfo& Server = Region.Region.Servers[ServerIdx];
+					const FString Address = FString::Printf(TEXT("%s:%d"), *Server.Address, Server.Port);
+
+					auto CompletionDelegate = [WeakThisCap, RegionId, NumTestsPerRegion, InCompletionDelegate](FIcmpEchoResult Result)
+					{
+						if (WeakThisCap.IsValid())
+						{
+							auto StrongThis = WeakThisCap.Get();
+							StrongThis->OnPingResultComplete(RegionId, NumTestsPerRegion, Result);
+							if (StrongThis->AreAllRegionsComplete())
+							{
+								StrongThis->CalculatePingAverages();
+								InCompletionDelegate.ExecuteIfBound(EQosCompletionResult::Success, StrongThis->Datacenters);
+								StrongThis->bInProgress = false;
+							}
+						}
+					};
+
+					UE_LOG(LogQos, Verbose, TEXT("Pinging [%s] %s"), *RegionId, *Address);
+					FUDPPing::UDPEcho(Address, InParams.Timeout, CompletionDelegate);
+					ServerIdx = (ServerIdx + 1) % NumServers;
+					bDidNothing = false;
+				}
+			}
+		}
+	}
+
+	if (bDidNothing)
+	{
+		InCompletionDelegate.ExecuteIfBound(EQosCompletionResult::Failure, Datacenters);
+	}
+}
+
+bool UQosEvaluator::FindQosServersByRegion(int32 RegionIdx, FOnQosSearchComplete InCompletionDelegate)
+{
+	bool bSuccess = false;
+	if (Datacenters.IsValidIndex(RegionIdx))
+	{
+		FQosRegionInfo& Datacenter = Datacenters[RegionIdx];
+		Datacenter.Reset();
+
+		IOnlineSubsystem* OnlineSub = Online::GetSubsystem(GetWorld());
+		if (OnlineSub)
+		{
+			IOnlineSessionPtr SessionInt = OnlineSub->GetSessionInterface();
+			if (SessionInt.IsValid())
+			{
+				if (!Datacenter.Region.RegionId.IsEmpty())
+				{
+					const TSharedRef<FOnlineSessionSearch> QosSearchParams = MakeShareable(new FOnlineSessionSearchQos());
+					QosSearchParams->QuerySettings.Set(SETTING_REGION, Datacenter.Region.RegionId, EOnlineComparisonOp::Equals);
+
+					QosSearchQuery = QosSearchParams;
+
+					FOnFindSessionsCompleteDelegate OnFindDatacentersCompleteDelegate = FOnFindSessionsCompleteDelegate::CreateUObject(this, &ThisClass::OnFindQosServersByRegionComplete, RegionIdx, InCompletionDelegate);
+					OnFindDatacentersCompleteDelegateHandle = SessionInt->AddOnFindSessionsCompleteDelegate_Handle(OnFindDatacentersCompleteDelegate);
+
+					SessionInt->FindSessions(ControllerId, QosSearchParams);
+				}
+				else
+				{
+					OnFindQosServersByRegionComplete(false, RegionIdx, InCompletionDelegate);
+				}
+				bSuccess = true;
+			}
+		}
+	}
+
+	return bSuccess;
+}
+
+int32 UQosEvaluator::GetNextRegionId(int32 LastRegionId)
+{
+	int32 NextIdx = LastRegionId + 1;
+	for (; Datacenters.IsValidIndex(NextIdx) && !Datacenters[NextIdx].Region.bEnabled; NextIdx++);
+	return NextIdx;
+}
+
+void UQosEvaluator::OnFindQosServersByRegionComplete(bool bWasSuccessful, int32 RegionIdx, FOnQosSearchComplete InCompletionDelegate)
+{
+	IOnlineSubsystem* OnlineSub = Online::GetSubsystem(GetWorld());
+	if (OnlineSub)
+	{
+		IOnlineSessionPtr SessionInt = OnlineSub->GetSessionInterface();
+		if (SessionInt.IsValid())
+		{
+			SessionInt->ClearOnFindSessionsCompleteDelegate_Handle(OnFindDatacentersCompleteDelegateHandle);
+		}
+	}
+
+	if (!bCancelOperation)
+	{
+		FQosRegionInfo& Datacenter = Datacenters[RegionIdx];
+		Datacenter.Result = bWasSuccessful ? EQosCompletionResult::Success : EQosCompletionResult::Failure;
+		
+		// Copy the search results for later evaluation
+		Datacenter.SearchResults = QosSearchQuery->SearchResults;
+		QosSearchQuery = nullptr;
+
+		int32 NextRegionIdx = GetNextRegionId(RegionIdx);
+		if (Datacenters.IsValidIndex(NextRegionIdx))
+		{
+			TWeakObjectPtr<UQosEvaluator> WeakThisCap(this);
+			GetWorldTimerManager().SetTimerForNextTick(FTimerDelegate::CreateLambda([InCompletionDelegate, NextRegionIdx, WeakThisCap]()
+			{
+				if (WeakThisCap.IsValid())
+				{
+					auto StrongThis = WeakThisCap.Get();
+					// Move to the next region results
+					if (!StrongThis->FindQosServersByRegion(NextRegionIdx, InCompletionDelegate))
+					{
+						// Failed to start
+						StrongThis->GetWorldTimerManager().SetTimerForNextTick(FTimerDelegate::CreateLambda([InCompletionDelegate, WeakThisCap]()
+						{
+							if (WeakThisCap.IsValid())
+							{
+								auto StrongThis1 = WeakThisCap.Get();
+								StrongThis1->FinalizeDatacenterResult(InCompletionDelegate, EQosCompletionResult::Failure, TArray<FQosRegionInfo>());
+							}
+						}));
+					}
+				}
+			}));
+		}
+		else
+		{
+			// Evaluate search results for all regions next tick
+			FOnQosPingEvalComplete CompletionDelegate = FOnQosPingEvalComplete::CreateUObject(this, &ThisClass::OnEvaluateForDatacenterComplete, InCompletionDelegate);
+		
+			TWeakObjectPtr<UQosEvaluator> WeakThisCap(this);
+			GetWorldTimerManager().SetTimerForNextTick(FTimerDelegate::CreateLambda([CompletionDelegate, WeakThisCap]()
+			{
+				if (WeakThisCap.IsValid())
+				{
+					auto StrongThis = WeakThisCap.Get();
+					StrongThis->bInProgress = false;
+					StrongThis->StartServerPing(CompletionDelegate);
+				}
+			}));
+		}
+	}
+	else
+	{
+		QosSearchQuery = nullptr;
+
+		// Mark all remaining datacenters as canceled
+		for (int32 Idx = RegionIdx; Idx < Datacenters.Num(); Idx++)
+		{
+			Datacenters[Idx].Result = EQosCompletionResult::Canceled;
+		}
+
+		FinalizeDatacenterResult(InCompletionDelegate, EQosCompletionResult::Canceled, Datacenters);
+	}
+}
+
+void UQosEvaluator::OnEvaluateForDatacenterComplete(EQosCompletionResult Result, FOnQosSearchComplete InCompletionDelegate)
+{
+	if (Result == EQosCompletionResult::Success)
+	{
+		//@TODO: This is a temporary measure to reduce the effect of framerate on data center ping estimation
+		// (due to the use of beacons that are ticked on the game thread for the estimation)
+		// This value can be 0 to disable discounting or something like 1 or 2
+		const float FractionOfFrameTimeToDiscount = 2.0f;
+
+		extern ENGINE_API float GAverageMS;
+		const int32 TimeToDiscount = (int32)(FractionOfFrameTimeToDiscount * GAverageMS);
+
+		CalculatePingAverages(TimeToDiscount);
+	}
+
+	FinalizeDatacenterResult(InCompletionDelegate, Result, Datacenters);
+}
+
+void UQosEvaluator::CalculatePingAverages(int32 TimeToDiscount)
+{
+	for (FQosRegionInfo& Datacenter : Datacenters)
+	{
+		int32 TotalPingInMs = 0;
+		int32 NumResults = 0;
+		for (int32 PingIdx = 0; PingIdx < Datacenter.PingResults.Num(); PingIdx++)
+		{
+			if (Datacenter.PingResults[PingIdx] != UNREACHABLE_PING)
+			{
+				TotalPingInMs += Datacenter.PingResults[PingIdx];
+				NumResults++;
+			}
+			else
+			{
+				UE_LOG(LogQos, Log, TEXT("Region[%s]: qos unreachable"), *Datacenter.Region.RegionId);
+			}
+		}
+
+		int32 RawAvgPing = UNREACHABLE_PING;
+		Datacenter.AvgPingMs = RawAvgPing;
+		if (NumResults > 0)
+		{
+			RawAvgPing = TotalPingInMs / NumResults;
+			Datacenter.AvgPingMs = FMath::Max<int32>(RawAvgPing - TimeToDiscount, 1);
+		}
+
+		UE_LOG(LogQos, Verbose, TEXT("Region[%s] Avg: %d Num: %d; Adjusted: %d"), *Datacenter.Region.RegionId, RawAvgPing, NumResults, Datacenter.AvgPingMs);
+
+		if (QosStats.IsValid())
+		{
+			QosStats->RecordRegionInfo(Datacenter.Region.RegionId, Datacenter.AvgPingMs, NumResults);
+		}
+	}
+}
+
+void UQosEvaluator::StartServerPing(const FOnQosPingEvalComplete& InCompletionDelegate)
+{
+	if (!bInProgress)
+	{
+		OnQosPingEvalComplete = InCompletionDelegate;
+
+		if (Datacenters.Num() > 0)
+		{
+			bInProgress = true;
+			EvaluateRegionPing(GetNextRegionId());
+		}
+		else
+		{
+			// No regions to ping
+			FinalizePingServers(EQosCompletionResult::Failure);
+		}
+	}
+	else
+	{
+		// Already in progress
+		InCompletionDelegate.ExecuteIfBound(EQosCompletionResult::Failure);
+	}
+}
+
+void UQosEvaluator::EvaluateRegionPing(int32 RegionIdx)
+{
+	if (ensure(bInProgress))
+	{
+		if (Datacenters.IsValidIndex(RegionIdx))
+		{
+			CurrentSearchPass.RegionIdx = RegionIdx;
+			CurrentSearchPass.CurrentSessionIdx = INDEX_NONE;
+			ContinuePingRegion();
+		}
+		else
+		{
+			// Invalid region id
+			FinalizePingServers(EQosCompletionResult::Failure);
+		}
+	}
+}
+
+void UQosEvaluator::ContinuePingRegion()
+{
+	if (!bCancelOperation)
+	{
+		bool bStartedPing = false;
+		if (Datacenters.IsValidIndex(CurrentSearchPass.RegionIdx))
+		{
+			FQosRegionInfo& Datacenter = Datacenters[CurrentSearchPass.RegionIdx];
+			CurrentSearchPass.CurrentSessionIdx++;
+			if (Datacenter.Result == EQosCompletionResult::Success)
+			{
+				if (Datacenter.SearchResults.IsValidIndex(CurrentSearchPass.CurrentSessionIdx))
+				{
+					bStartedPing = true;
+
+					// There are more valid search results, keep attempting Qos requests
+					UWorld* World = GetWorld();
+					QosBeaconClient = World->SpawnActor<AQosBeaconClient>(AQosBeaconClient::StaticClass());
+
+					QosBeaconClient->OnQosRequestComplete().BindUObject(this, &ThisClass::OnQosRequestComplete);
+					QosBeaconClient->OnHostConnectionFailure().BindUObject(this, &ThisClass::OnQosConnectionFailure);
+					QosBeaconClient->SendQosRequest(Datacenter.SearchResults[CurrentSearchPass.CurrentSessionIdx]);
+				}
+				else
+				{
+					// Out of search results
+					Datacenter.LastCheckTimestamp = FDateTime::UtcNow();
+				}
+			}
+		}
+			
+		if (!bStartedPing)
+		{
+			int32 NextRegionIdx = GetNextRegionId(CurrentSearchPass.RegionIdx);
+			if (Datacenters.IsValidIndex(NextRegionIdx))
+			{
+				// Ran out of search results, advance regions
+				EvaluateRegionPing(NextRegionIdx);
+			}
+			else
+			{
+				// Completely done
+				FinalizePingServers(EQosCompletionResult::Success);
+			}
+		}
+	}
+	else
+	{
+		// Operation canceled
+		FinalizePingServers(EQosCompletionResult::Canceled);
+	}
+}
+
+void UQosEvaluator::OnQosRequestComplete(EQosResponseType QosResponse, int32 ResponseTime)
+{
+	if (QosBeaconClient.IsValid())
+	{
+		DestroyClientBeacons();
+	}
+
+	if (ensure(Datacenters.IsValidIndex(CurrentSearchPass.RegionIdx)))
+	{
+		FQosRegionInfo& Datacenter = Datacenters[CurrentSearchPass.RegionIdx];
+		ensure(Datacenter.SearchResults.IsValidIndex(CurrentSearchPass.CurrentSessionIdx));
+		FOnlineSessionSearchResult& SearchResult = Datacenter.SearchResults[CurrentSearchPass.CurrentSessionIdx];
+		if (QosResponse == EQosResponseType::Success)
+		{
+			Datacenter.PingResults.Add(ResponseTime);
+			SearchResult.PingInMs = ResponseTime;
+		}
+		else
+		{
+			Datacenter.PingResults.Add(UNREACHABLE_PING);
+			SearchResult.PingInMs = UNREACHABLE_PING;
+		}
+
+		extern ENGINE_API float GAverageFPS;
+		extern ENGINE_API float GAverageMS;
+		UE_LOG(LogQos, Verbose, TEXT("Qos response received for region %s: %d ms FPS: %0.2f MS: %0.2f"), *Datacenter.Region.RegionId, ResponseTime, GAverageFPS, GAverageMS);
+
+		if (QosStats.IsValid())
+		{
+			QosStats->RecordQosAttempt(SearchResult, QosResponse == EQosResponseType::Success);
+		}
+	}
+
+	// Cancel operation will occur next tick if applicable
+	GetWorldTimerManager().SetTimerForNextTick(FTimerDelegate::CreateUObject(this, &ThisClass::ContinuePingRegion));
+}
+
+void UQosEvaluator::OnQosConnectionFailure()
+{
+	OnQosRequestComplete(EQosResponseType::Failure, UNREACHABLE_PING);
+}
+
+void UQosEvaluator::FinalizePingServers(EQosCompletionResult Result)
+{
+	UE_LOG(LogQos, Log, TEXT("Ping evaluation complete. Result:%s"), ToString(Result));
+
+	TWeakObjectPtr<UQosEvaluator> WeakThisCap(this);
+	FTimerDelegate TimerDelegate = FTimerDelegate::CreateLambda([Result, WeakThisCap]() {
+		if (WeakThisCap.IsValid())
+		{
+			auto StrongThis = WeakThisCap.Get();
+			StrongThis->OnQosPingEvalComplete.ExecuteIfBound(Result);
+			StrongThis->ResetSearchVars();
+		}
+	});
+	GetWorldTimerManager().SetTimerForNextTick(TimerDelegate);
+}
+
+bool UQosEvaluator::AreAllRegionsComplete()
+{
+	for (FQosRegionInfo& Region : Datacenters)
+	{
+		if (Region.Region.bEnabled && Region.Result == EQosCompletionResult::Invalid)
+		{
+			return false;
+		}
+	}
+
+	return true;
+}
+
+void UQosEvaluator::OnPingResultComplete(const FString& RegionId, int32 NumTests, const FIcmpEchoResult& Result)
+{
+	for (FQosRegionInfo& Region : Datacenters)
+	{
+		if (Region.Region.RegionId == RegionId)
+		{
+			UE_LOG(LogQos, Verbose, TEXT("Ping Complete [%s] %s: %d"), *RegionId, *Result.ResolvedAddress, (int32)(Result.Time * 1000.0f));
+
+			const bool bSuccess = (Result.Status == EIcmpResponseStatus::Success);
+			int32 PingInMs = bSuccess ? (int32)(Result.Time * 1000.0f) : UNREACHABLE_PING;
+			Region.PingResults.Add(PingInMs);
+
+			if (QosStats.IsValid())
+			{
+				QosStats->RecordQosAttempt(RegionId, PingInMs, bSuccess);
+			}
+
+			if (Region.PingResults.Num() == NumTests)
+			{
+				Region.LastCheckTimestamp = FDateTime::UtcNow();
+				Region.Result = EQosCompletionResult::Success;
+			}
+			break;
+		}
+	}
+}
+
+void UQosEvaluator::ResetSearchVars()
+{
+	bInProgress = false;
+	bCancelOperation = false;
+	CurrentSearchPass.RegionIdx = INDEX_NONE;
+	CurrentSearchPass.CurrentSessionIdx = INDEX_NONE;
+	OnQosPingEvalComplete.Unbind();
+	QosSearchQuery = nullptr;
+}
+
+void UQosEvaluator::DestroyClientBeacons()
+{
+	if (QosBeaconClient.IsValid())
+	{
+		QosBeaconClient->OnQosRequestComplete().Unbind();
+		QosBeaconClient->OnHostConnectionFailure().Unbind();
+		QosBeaconClient->DestroyBeacon();
+		QosBeaconClient.Reset();
+	}
+}
+
+void UQosEvaluator::StartAnalytics()
+{
+	if (AnalyticsProvider.IsValid())
+	{
+		ensure(!QosStats.IsValid());
+		QosStats = MakeShareable(new FQosDatacenterStats());
+		QosStats->StartQosPass();
+	}
+}
+
+void UQosEvaluator::EndAnalytics(EQosCompletionResult CompletionResult)
+{
+	if (QosStats.IsValid())
+	{
+		if (CompletionResult != EQosCompletionResult::Canceled)
+		{
+			EDatacenterResultType ResultType = EDatacenterResultType::Failure;
+			if (CompletionResult != EQosCompletionResult::Failure)
+			{
+				ResultType = EDatacenterResultType::Normal;
+			}
+
+			QosStats->EndQosPass(ResultType);
+			QosStats->Upload(AnalyticsProvider);
+		}
+
+		QosStats = nullptr;
+	}
+}
+
+UWorld* UQosEvaluator::GetWorld() const
+{
+	UWorld* World = ParentWorld.Get();
+	check(World);
+	return World;
+}
+
+FTimerManager& UQosEvaluator::GetWorldTimerManager() const
+{
+	return GetWorld()->GetTimerManager();
+}
+
diff --git a/Engine/Plugins/Online/OnlineFramework/Source/Qos/Private/QosEvaluator.h b/Engine/Plugins/Online/OnlineFramework/Source/Qos/Private/QosEvaluator.h
new file mode 100644
--- /dev/null
+++ b/Engine/Plugins/Online/OnlineFramework/Source/Qos/Private/QosEvaluator.h
@@ -0,0 +1,319 @@
+// Copyright 1998-2016 Epic Games, Inc. All Rights Reserved.
+
+#pragma once
+
+#include "OnlineSessionInterface.h"
+#include "QosRegionManager.h"
+#include "QosEvaluator.generated.h"
+
+class AQosBeaconClient;
+class FQosDatacenterStats;
+struct FIcmpEchoResult;
+enum class EQosResponseType : uint8;
+
+#define GAMEMODE_QOS	TEXT("QOSSERVER")
+
+/**
+ * Search settings for QoS advertised sessions
+ */
+class QOS_API FOnlineSessionSearchQos : public FOnlineSessionSearch
+{
+public:
+	FOnlineSessionSearchQos();
+	virtual ~FOnlineSessionSearchQos() {}
+
+	// FOnlineSessionSearch Interface begin
+	virtual TSharedPtr<FOnlineSessionSettings> GetDefaultSessionSettings() const override;
+	virtual void SortSearchResults() override;
+	// FOnlineSessionSearch Interface end
+};
+
+/**
+ * Internal state for a given QoS pass
+ */
+USTRUCT()
+struct FQosSearchPass
+{
+	GENERATED_USTRUCT_BODY()
+
+	/** Current Region index in the datacenter array */
+	UPROPERTY()
+	int32 RegionIdx;
+	/** Current search result choice to test */
+	UPROPERTY()
+	int32 CurrentSessionIdx;
+	/** Current array of QoS search results to evaluate */
+	//TArray<FOnlineSessionSearchResult> SearchResults;
+
+	FQosSearchPass()
+		: RegionIdx(INDEX_NONE)
+		, CurrentSessionIdx(INDEX_NONE)
+	{}
+};
+
+/**
+ * Input parameters to start a qos ping check
+ */
+struct FQosParams
+{
+	/** User that initiated the request */
+	int32 ControllerId;
+	/** Use the old qos beacon method for ping determination */
+	bool bUseOldQosServers;
+	/** Number of ping requests per region */
+	int32 NumTestsPerRegion;
+	/** Amount of time to wait for each request */
+	float Timeout;
+};
+
+/*
+ * Delegate triggered when an evaluation of ping for all servers in a search query have completed
+ *
+ * @param Result the ping operation result
+ */
+DECLARE_DELEGATE_OneParam(FOnQosPingEvalComplete, EQosCompletionResult /** Result */);
+
+/** 
+ * Delegate triggered when all QoS search results have been investigated 
+ *
+ * @param Result the QoS operation result
+ * @param RegionInfo The per-region ping information
+ */
+DECLARE_DELEGATE_TwoParams(FOnQosSearchComplete, EQosCompletionResult /** Result */, const TArray<FQosRegionInfo>& /** RegionInfo */);
+
+/**
+ * Evaluates QoS metrics to determine the best datacenter under current conditions
+ * Additionally capable of generically pinging an array of servers that have a QosBeaconHost active
+ */
+UCLASS(config = Engine)
+class QOS_API UQosEvaluator : public UObject
+{
+	GENERATED_UCLASS_BODY()
+
+public:
+
+	/**
+	 * QoS services
+	 */
+
+	/**
+	 * Find all the advertised datacenters and begin the process of evaluating ping results
+	 * Will return the default datacenter in the event of failure or no advertised datacenters
+	 *
+	 * @param InParams parameters defining the request
+	 * @param InDatacenters array of datacenters to query
+	 * @param InCompletionDelegate delegate to fire when a datacenter choice has been made
+	 */
+	void FindDatacenters(const FQosParams& InParams, TArray<FQosDatacenterInfo>& InDatacenters, const FOnQosSearchComplete& InCompletionDelegate);
+
+	/**
+	 * Is a QoS operation active
+	 *
+	 * @return true if QoS is active, false otherwise
+	 */
+	bool IsActive() const { return bInProgress; }
+
+	/**
+	 * Cancel the current QoS operation at the earliest opportunity
+	 */
+	void Cancel();
+
+	void SetWorld(UWorld* InWorld);
+
+protected:
+
+	/**
+	 * Use the udp ping code to ping known servers
+	 *
+	 * @param InParams parameters defining the request
+	 * @param InCompletionDelegate delegate to fire when all regions have completed their tests
+	 */
+	void PingRegionServers(const FQosParams& InParams, const FOnQosSearchComplete& InCompletionDelegate);
+
+	/**
+	 * Given a single region id, find available qos servers
+	 * 
+	 * @param RegionIdx region index in the array
+	 * @param InCompletionDelegate delegate to fire when complete
+	 */
+	bool FindQosServersByRegion(int32 RegionIdx, FOnQosSearchComplete InCompletionDelegate);
+
+	/**
+	 * Start pinging all servers found across all regions, starting with the first datacenter
+	 *
+	 * @param InCompletionDelegate delegate to fire when complete
+	 */
+	void StartServerPing(const FOnQosPingEvalComplete& InCompletionDelegate);
+
+	/**
+	 * Evaluate ping reachability for a given list of servers.  Assumes there is a AQosBeaconHost on the machine
+	 * to receive the request and send a response
+	 *
+	 * @param RegionIdx region to evaluate qos servers 
+	 */
+	void EvaluateRegionPing(int32 RegionIdx);
+
+private:
+
+	/** Current QoS search/eval state */
+	UPROPERTY()
+	FQosSearchPass CurrentSearchPass;
+
+	/** Reference to external UWorld */
+	TWeakObjectPtr<UWorld> ParentWorld;
+
+	/** QoS search results */
+	TSharedPtr<FOnlineSessionSearch> QosSearchQuery;
+
+	/** Delegate for finding datacenters */
+	FOnFindSessionsCompleteDelegate OnQosSessionsSearchCompleteDelegate;
+	FOnQosPingEvalComplete OnQosPingEvalComplete;
+
+	/** Beacon for sending QoS requests */
+	TWeakObjectPtr<AQosBeaconClient> QosBeaconClient;
+
+	/** User initiating the request */
+	UPROPERTY()
+	int32 ControllerId;
+	/** A QoS operation is in progress */
+	UPROPERTY()
+	bool bInProgress;
+	/** Should cancel occur at the next available opportunity */
+	UPROPERTY()
+	bool bCancelOperation;
+
+	/** Array of datacenters currently being evaluated */
+	UPROPERTY(Transient)
+	TArray<FQosRegionInfo> Datacenters;
+
+	/**
+	 * Signal QoS operation has completed next frame
+	 *
+	 * @param InCompletionDelegate the delegate to trigger
+	 * @param CompletionResult the QoS operation result
+	 */
+	void FinalizeDatacenterResult(const FOnQosSearchComplete& InCompletionDelegate, EQosCompletionResult CompletionResult, const TArray<FQosRegionInfo>& RegionInfo);
+
+	/**
+	 * Continue with the next datacenter endpoint and gather its ping result
+	 */
+	void ContinuePingRegion();
+
+	/**
+	 * Finalize the search results from pinging various datacenter endpoints to find the best datacenter
+	 *
+	 * @param Result the QoS operation result
+	 */
+	void FinalizePingServers(EQosCompletionResult Result);
+
+	/**
+	 * Delegate fired when all servers within a given search query have been evaluated
+	 *
+	 * @param Result the QoS operation result
+	 * @param InCompletionDelegate the datacenter complete delegate to fire
+	 */
+	void OnEvaluateForDatacenterComplete(EQosCompletionResult Result, FOnQosSearchComplete InCompletionDelegate);
+
+	/**
+	 * Delegate fired when a datacenter search query has completed
+	 *
+	 * @param bWasSuccessful true if the async action completed without error, false if there was an error
+	 * @param RegionId region search performed
+	 * @param InCompletionDelegate the datacenter complete delegate to fire
+	 */
+	void OnFindQosServersByRegionComplete(bool bWasSuccessful, int32 RegionIdx, FOnQosSearchComplete InCompletionDelegate);
+
+	/**
+	 * Delegate from the QoS beacon fired with each beacon response
+	 *
+	 * @param QosResponse the response from the server, success or otherwise
+	 * @param ResponseTime the time to complete the request, or MAX_QUERY_PING if unsuccessful
+	 */
+	void OnQosRequestComplete(EQosResponseType QosResponse, int32 ResponseTime);
+
+	/**
+	 * Delegate from the QoS beacon fired when there is a failure to connect to an endpoint
+	 */
+	void OnQosConnectionFailure();
+
+	/**
+	 * @return true if all ping requests have completed (new method)
+	 */
+	bool AreAllRegionsComplete();
+
+	void OnPingResultComplete(const FString& RegionId, int32 NumTests, const FIcmpEchoResult& Result);
+
+	/**
+	 * Take all found ping results and process them before consumption at higher levels
+	 *
+	 * @param TimeToDiscount amount of time to subtract from calculation to compensate for external factors (frame rate, etc)
+	 */
+	void CalculatePingAverages(int32 TimeToDiscount = 0);
+
+private:
+
+	FDelegateHandle OnFindDatacentersCompleteDelegateHandle;
+
+public:
+
+	/**
+	 * Analytics
+	 */
+
+	void SetAnalyticsProvider(TSharedPtr<IAnalyticsProvider> InAnalyticsProvider);
+
+private:
+
+	void StartAnalytics();
+	void EndAnalytics(EQosCompletionResult CompletionResult);
+
+	/** Reference to the provider to submit data to */
+	TSharedPtr<IAnalyticsProvider> AnalyticsProvider;
+	/** Stats related to these operations */
+	TSharedPtr<FQosDatacenterStats> QosStats;
+
+private:
+
+	/**
+	 * Helpers
+	 */
+
+	int32 GetNextRegionId(int32 LastRegionId = INDEX_NONE);
+
+	/** Reset all variables related to the QoS evaluator */
+	void ResetSearchVars();
+
+	/** Quick access to the current world */
+	UWorld* GetWorld() const;
+
+	/** Quick access to the world timer manager */
+	FTimerManager& GetWorldTimerManager() const;
+
+	/** Client beacon cleanup */
+	void DestroyClientBeacons();
+};
+
+inline const TCHAR* ToString(EQosCompletionResult Result)
+{
+	switch (Result)
+	{
+		case EQosCompletionResult::Invalid:
+		{
+			return TEXT("Invalid");
+		}
+		case EQosCompletionResult::Success:
+		{
+			return TEXT("Success");
+		}
+		case EQosCompletionResult::Failure:
+		{
+			return TEXT("Failure");
+		}
+		case EQosCompletionResult::Canceled:
+		{
+			return TEXT("Canceled");
+		}
+	}
+
+	return TEXT("");
+}
