Integrate as edit CL#4113476:
- Fix a rare GC assertion/crash with Qos

diff --git a/Engine/Plugins/Online/OnlineFramework/Source/Qos/Private/QosRegionManager.cpp b/Engine/Plugins/Online/OnlineFramework/Source/Qos/Private/QosRegionManager.cpp
--- a/Engine/Plugins/Online/OnlineFramework/Source/Qos/Private/QosRegionManager.cpp
+++ b/Engine/Plugins/Online/OnlineFramework/Source/Qos/Private/QosRegionManager.cpp
@@ -170,6 +170,7 @@ void UQosRegionManager::BeginQosEvaluation(UWorld* World, const TSharedPtr<IAnal
 
 	// create a new evaluator and start the process of running
 	Evaluator = NewObject<UQosEvaluator>();
+	Evaluator->AddToRoot();
 	Evaluator->SetWorld(World);
 	Evaluator->SetAnalyticsProvider(AnalyticsProvider);
 
@@ -184,7 +185,11 @@ void UQosRegionManager::BeginQosEvaluation(UWorld* World, const TSharedPtr<IAnal
 void UQosRegionManager::OnQosEvaluationComplete(EQosCompletionResult Result, const TArray<FQosRegionInfo>& RegionInfo)
 {
 	// toss the evaluator
-	Evaluator = nullptr;
+	if (Evaluator != nullptr)
+	{
+		Evaluator->RemoveFromRoot();
+		Evaluator = nullptr;
+	}
 	QosEvalResult = Result;
 	RegionOptions.Empty(RegionInfo.Num());
 
