Mark FPartyState as a BlueprintType
#rnx

diff --git a/Engine/Plugins/Online/OnlineFramework/Source/Party/Public/PartyGameState.h b/Engine/Plugins/Online/OnlineFramework/Source/Party/Public/PartyGameState.h
--- a/Engine/Plugins/Online/OnlineFramework/Source/Party/Public/PartyGameState.h
+++ b/Engine/Plugins/Online/OnlineFramework/Source/Party/Public/PartyGameState.h
@@ -97,7 +97,7 @@ enum class ECrossplayPreference: uint8
 /**
  * Current state of the party
  */
-USTRUCT()
+USTRUCT(BlueprintType)
 struct FPartyState
 {
 	GENERATED_USTRUCT_BODY();
