Const usage of IOnlinePartyJoinInfo

diff --git a/Engine/Plugins/Online/OnlineFramework/Source/Party/Private/SocialManager.cpp b/Engine/Plugins/Online/OnlineFramework/Source/Party/Private/SocialManager.cpp
--- a/Engine/Plugins/Online/OnlineFramework/Source/Party/Private/SocialManager.cpp
+++ b/Engine/Plugins/Online/OnlineFramework/Source/Party/Private/SocialManager.cpp
@@ -757,7 +757,7 @@ USocialParty* USocialManager::GetPartyInternal(const FOnlinePartyId& PartyId, bo
 	return nullptr;
 }
 
-TSharedPtr<IOnlinePartyJoinInfo> USocialManager::GetJoinInfoFromSession(const FOnlineSessionSearchResult& PlatformSession)
+TSharedPtr<const IOnlinePartyJoinInfo> USocialManager::GetJoinInfoFromSession(const FOnlineSessionSearchResult& PlatformSession)
 {
 	static const FName JoinInfoSettingName = PLATFORM_XBOXONE ? SETTING_CUSTOM_JOIN_INFO : SETTING_CUSTOM;
 
diff --git a/Engine/Plugins/Online/OnlineFramework/Source/Party/Public/SocialManager.h b/Engine/Plugins/Online/OnlineFramework/Source/Party/Public/SocialManager.h
--- a/Engine/Plugins/Online/OnlineFramework/Source/Party/Public/SocialManager.h
+++ b/Engine/Plugins/Online/OnlineFramework/Source/Party/Public/SocialManager.h
@@ -188,7 +188,7 @@ private:
 	USocialParty* GetPartyInternal(const FOnlinePartyTypeId& PartyTypeId, bool bIncludeLeavingParties = false) const;
 	USocialParty* GetPartyInternal(const FOnlinePartyId& PartyId, bool bIncludeLeavingParties = false) const;
 
-	TSharedPtr<IOnlinePartyJoinInfo> GetJoinInfoFromSession(const FOnlineSessionSearchResult& PlatformSession);
+	TSharedPtr<const IOnlinePartyJoinInfo> GetJoinInfoFromSession(const FOnlineSessionSearchResult& PlatformSession);
 
 private:	// Handlers
 	void HandleGameViewportInitialized();
diff --git a/Engine/Plugins/Online/OnlineFramework/Source/Party/Public/User/SocialUser.h b/Engine/Plugins/Online/OnlineFramework/Source/Party/Public/User/SocialUser.h
--- a/Engine/Plugins/Online/OnlineFramework/Source/Party/Public/User/SocialUser.h
+++ b/Engine/Plugins/Online/OnlineFramework/Source/Party/Public/User/SocialUser.h
@@ -165,7 +165,7 @@ private:
 	int32 NumPendingQueries = 0;
 	bool bIsInitialized = false;
 
-	TSharedPtr<IOnlinePartyJoinInfo> PersistentPartyInfo;
+	TSharedPtr<const IOnlinePartyJoinInfo> PersistentPartyInfo;
 	TMap<ESocialSubsystem, FSubsystemUserInfo> SubsystemInfoByType;
 
 	// Initialization delegates that fire only when a specific user has finishing initializing
