Add FOnlineParty::GetConfiguration in preparation for deprecating the non-const access to FOnlineParty::Config

diff --git a/Engine/Plugins/Online/OnlineFramework/Source/Party/Private/Party/SocialParty.cpp b/Engine/Plugins/Online/OnlineFramework/Source/Party/Private/Party/SocialParty.cpp
--- a/Engine/Plugins/Online/OnlineFramework/Source/Party/Private/Party/SocialParty.cpp
+++ b/Engine/Plugins/Online/OnlineFramework/Source/Party/Private/Party/SocialParty.cpp
@@ -374,7 +374,7 @@ void USocialParty::InitializeParty(const TSharedRef<const FOnlineParty>& InOssPa
 		PartyDataReplicator->SetOwningParty(*this);
 
 		OssParty = InOssParty;
-		CurrentConfig = *InOssParty->Config;
+		CurrentConfig = *InOssParty->GetConfiguration();
 		CurrentLeaderId = InOssParty->LeaderId;
 
 		OwningLocalUserId = GetSocialManager().GetFirstLocalUserId(ESocialSubsystem::Primary);
@@ -671,7 +671,7 @@ void USocialParty::HandlePartyConfigChanged(const FUniqueNetId& LocalUserId, con
 {
 	if (PartyId == GetPartyId())
 	{
-		CurrentConfig = *OssParty->Config;
+		CurrentConfig = *OssParty->GetConfiguration();
 		OnPartyConfigurationChanged().Broadcast(CurrentConfig);
 	}
 }
@@ -682,7 +682,7 @@ void USocialParty::HandleUpdatePartyConfigComplete(const FUniqueNetId& LocalUser
 	{
 		UE_LOG(LogParty, Verbose, TEXT("[%s] Party config updated %s"), *PartyId.ToDebugString(), ToString(Result));
 
-		CurrentConfig = *OssParty->Config;
+		CurrentConfig = *OssParty->GetConfiguration();
 		OnPartyConfigurationChanged().Broadcast(CurrentConfig);
 	}
 	else
@@ -982,7 +982,7 @@ void USocialParty::SetPartyMaxSize(int32 NewSize)
 int32 USocialParty::GetPartyMaxSize() const
 {
 	check(OssParty.IsValid());
-	return OssParty->Config->MaxMembers;
+	return OssParty->GetConfiguration()->MaxMembers;
 }
 
 FPartyJoinDenialReason USocialParty::GetPublicJoinability() const
diff --git a/Engine/Plugins/Online/OnlineSubsystem/Source/Public/Interfaces/OnlinePartyInterface.h b/Engine/Plugins/Online/OnlineSubsystem/Source/Public/Interfaces/OnlinePartyInterface.h
--- a/Engine/Plugins/Online/OnlineSubsystem/Source/Public/Interfaces/OnlinePartyInterface.h
+++ b/Engine/Plugins/Online/OnlineSubsystem/Source/Public/Interfaces/OnlinePartyInterface.h
@@ -503,9 +503,28 @@ protected:
 public:
 	virtual ~FOnlineParty() = default;
 
+	/**
+	 * Check if the local user has invite permissions in this party. Based on configuration permissions and party state.
+	 *
+	 * @param LocalUserId the local user's id
+	 * @return true if the local user can invite, false if not
+	 */
 	virtual bool CanLocalUserInvite(const FUniqueNetId& LocalUserId) const = 0;
+
+	/**
+	 * Is this party joinable?
+	 *
+	 * @return true if this party is joinable, false if not
+	 */
 	virtual bool IsJoinable() const = 0;
 
+	/**
+	 * Get the party's configuration
+	 *
+	 * @return the party's configuration
+	 */
+	virtual TSharedRef<const FPartyConfiguration> GetConfiguration() const = 0;
+
 	/** Unique id of the party */
 	TSharedRef<const FOnlinePartyId> PartyId;
 	/** Type of party (e.g., Primary) */
