Social User: Can send friend invite checks if passed in user is local, if so then false.
Also avoid a potential crash in GetFriendPresenceInfo.

diff --git a/Engine/Plugins/Online/OnlineFramework/Source/Party/Private/User/SocialUser.cpp b/Engine/Plugins/Online/OnlineFramework/Source/Party/Private/User/SocialUser.cpp
--- a/Engine/Plugins/Online/OnlineFramework/Source/Party/Private/User/SocialUser.cpp
+++ b/Engine/Plugins/Online/OnlineFramework/Source/Party/Private/User/SocialUser.cpp
@@ -442,7 +442,8 @@ const FOnlineUserPresence* USocialUser::GetFriendPresenceInfo(ESocialSubsystem S
 		if (IOnlinePresencePtr PresenceInterface = SocialOss ? SocialOss->GetPresenceInterface() : nullptr)
 		{
 			TSharedPtr<FOnlineUserPresence> LocalUserPresence;
-			if (PresenceInterface->GetCachedPresence(*GetUserId(SubsystemType), LocalUserPresence) == EOnlineCachedResult::Success)
+			const FUniqueNetIdRepl& UserId = GetUserId(SubsystemType);
+			if (UserId.IsValid() && PresenceInterface->GetCachedPresence(*UserId, LocalUserPresence) == EOnlineCachedResult::Success)
 			{
 				return LocalUserPresence.Get();
 			}
@@ -671,7 +672,7 @@ bool USocialUser::CanSendFriendInvite(ESocialSubsystem SubsystemType) const
 		}
 	}
 
-	return HasSubsystemInfo(SubsystemType) && !IsFriend(SubsystemType) && !IsBlocked(SubsystemType) && !IsFriendshipPending(SubsystemType);
+	return HasSubsystemInfo(SubsystemType) && !IsLocalUser() && !IsFriend(SubsystemType) && !IsBlocked(SubsystemType) && !IsFriendshipPending(SubsystemType);
 }
 
 void USocialUser::JoinParty(const FOnlinePartyTypeId& PartyTypeId) const
