Fixing compile error introduced in 5316175.

#tests Compile DebugGame Editor Win64

diff --git a/Engine/Plugins/Online/OnlineFramework/Source/Party/Private/Party/SocialParty.cpp b/Engine/Plugins/Online/OnlineFramework/Source/Party/Private/Party/SocialParty.cpp
--- a/Engine/Plugins/Online/OnlineFramework/Source/Party/Private/Party/SocialParty.cpp
+++ b/Engine/Plugins/Online/OnlineFramework/Source/Party/Private/Party/SocialParty.cpp
@@ -160,7 +160,7 @@ FPartyJoinApproval USocialParty::EvaluateJIPRequest(const FUniqueNetId& PlayerId
 	return JoinApproval;
 }
 
-bool USocialParty::ApplyCrossplayRestriction(FPartyJoinApproval& JoinApproval, const FOnlinePartyData& JoinData) const
+bool USocialParty::ApplyCrossplayRestriction(FPartyJoinApproval& JoinApproval, const FUserPlatform& Platform, const FOnlinePartyData& JoinData) const
 {
 	const ECrossplayPreference SenderCrossplayPreference = GetCrossplayPreferenceFromJoinData(JoinData);
 	const bool bSenderAllowsCrossplay = !OptedOutOfCrossplay(SenderCrossplayPreference);
@@ -191,7 +191,7 @@ bool USocialParty::ApplyCrossplayRestriction(FPartyJoinApproval& JoinApproval, c
 		}
 	}
 
-	return JoinApproval.DenialReason.HasAnyReason();
+	return JoinApproval.GetDenialReason().HasAnyReason();
 }
 
 FPartyJoinApproval USocialParty::EvaluateJoinRequest(const FUniqueNetId& PlayerId, const FUserPlatform& Platform, const FOnlinePartyData& JoinData, bool bFromJoinRequest) const
@@ -213,7 +213,7 @@ FPartyJoinApproval USocialParty::EvaluateJoinRequest(const FUniqueNetId& PlayerI
 		{
 			if (bAlwaysCheckCrossplatformOnPartyJoin)
 			{
-				ApplyCrossplayRestriction(JoinApproval, JoinData);
+				ApplyCrossplayRestriction(JoinApproval, Platform, JoinData);
 			}
 		}
 
diff --git a/Engine/Plugins/Online/OnlineFramework/Source/Party/Public/Party/SocialParty.h b/Engine/Plugins/Online/OnlineFramework/Source/Party/Public/Party/SocialParty.h
--- a/Engine/Plugins/Online/OnlineFramework/Source/Party/Public/Party/SocialParty.h
+++ b/Engine/Plugins/Online/OnlineFramework/Source/Party/Public/Party/SocialParty.h
@@ -230,7 +230,7 @@ protected:
 	/** Override in child classes to specify the type of UPartyMember to create */
 	virtual TSubclassOf<UPartyMember> GetDesiredMemberClass(bool bLocalPlayer) const;
 
-	bool ApplyCrossplayRestriction(FPartyJoinApproval& JoinApproval, const FOnlinePartyData& JoinData) const;
+	bool ApplyCrossplayRestriction(FPartyJoinApproval& JoinApproval, const FUserPlatform& Platform, const FOnlinePartyData& JoinData) const;
 	FName GetGameSessionName() const;
 	bool IsInRestrictedGameSession() const;
 
