Unshelved from pending changelist '5199188':

Party V1 modification to support testing V2 integration.

diff --git a/Engine/Plugins/Online/OnlineFramework/Source/Party/Private/Party/PartyMember.cpp b/Engine/Plugins/Online/OnlineFramework/Source/Party/Private/Party/PartyMember.cpp
--- a/Engine/Plugins/Online/OnlineFramework/Source/Party/Private/Party/PartyMember.cpp
+++ b/Engine/Plugins/Online/OnlineFramework/Source/Party/Private/Party/PartyMember.cpp
@@ -64,6 +64,7 @@ void UPartyMember::InitializePartyMember(const TSharedRef<FOnlinePartyMember>& I
 	if (ensure(!OssPartyMember.IsValid()))
 	{
 		OssPartyMember = InOssMember;
+		OssPartyMember->OnMemberConnectionStatusChangedDelegates.AddUObject(this, &ThisClass::HandleMemberConnectionStatusChanged);
 		USocialToolkit* OwnerToolkit = GetParty().GetSocialManager().GetFirstLocalUserToolkit();
 		check(OwnerToolkit);
 
@@ -253,4 +254,18 @@ void UPartyMember::HandleSocialUserInitialized(USocialUser& InitializedUser)
 	{
 		FinishInitializing();
 	}
-}
\ No newline at end of file
+}
+
+EMemberConnectionStatus UPartyMember::GetMemberConnectionStatus() const
+{
+	if (OssPartyMember.IsValid())
+	{
+		return OssPartyMember->MemberConnectionStatus;
+	}
+	return EMemberConnectionStatus::Uninitialized;
+}
+
+void UPartyMember::HandleMemberConnectionStatusChanged(const FUniqueNetId& ChangedUserId, const EMemberConnectionStatus NewMemberConnectionStatus, const EMemberConnectionStatus PreviousMemberConnectionStatus)
+{
+	OnMemberConnectionStatusChanged().Broadcast();
+}
diff --git a/Engine/Plugins/Online/OnlineFramework/Source/Party/Private/Party/SocialParty.cpp b/Engine/Plugins/Online/OnlineFramework/Source/Party/Private/Party/SocialParty.cpp
--- a/Engine/Plugins/Online/OnlineFramework/Source/Party/Private/Party/SocialParty.cpp
+++ b/Engine/Plugins/Online/OnlineFramework/Source/Party/Private/Party/SocialParty.cpp
@@ -416,7 +416,7 @@ void USocialParty::InitializePartyInternal()
 	PartyInterface->AddOnPartyMemberDataReceivedDelegate_Handle(FOnPartyMemberDataReceivedDelegate::CreateUObject(this, &USocialParty::HandlePartyMemberDataReceived));
 	PartyInterface->AddOnPartyMemberPromotedDelegate_Handle(FOnPartyMemberPromotedDelegate::CreateUObject(this, &USocialParty::HandlePartyMemberPromoted));
 	PartyInterface->AddOnPartyMemberExitedDelegate_Handle(FOnPartyMemberExitedDelegate::CreateUObject(this, &USocialParty::HandlePartyMemberExited));
-
+	//PartyInterface->AddOnPartyMemberStateChangedDelegate_Handle(FOnPartyMemberStateChangedDelegate::CreateUObject(this, &USocialParty::HandlePartyMemberStateChanged));
 	// Create a UPartyMember for every existing member on the OSS party
 	TArray<TSharedRef<FOnlinePartyMember>> OssPartyMembers;
 	PartyInterface->GetPartyMembers(*OwningLocalUserId, GetPartyId(), OssPartyMembers);
@@ -424,7 +424,7 @@ void USocialParty::InitializePartyInternal()
 	{
 		GetOrCreatePartyMember(*OssMember->GetUserId());
 	}
-	HandlePartyStateChanged(*OwningLocalUserId, GetPartyId(), OssParty->State);
+	HandlePartyStateChanged(*OwningLocalUserId, GetPartyId(), OssParty->State, OssParty->PreviousState);
 
 	if (IsLocalPlayerPartyLeader())
 	{
@@ -939,6 +939,12 @@ EPartyState USocialParty::GetOssPartyState() const
 	return OssParty->State;
 }
 
+EPartyState USocialParty::GetOssPartyPreviousState() const
+{
+	check(OssParty.IsValid());
+	return OssParty->PreviousState;
+}
+
 bool USocialParty::IsCurrentlyCrossplaying() const
 {
 	TArray<FUserPlatform> AllPlatformsPresent;
@@ -1125,7 +1131,7 @@ TSubclassOf<UPartyMember> USocialParty::GetDesiredMemberClass(bool bLocalPlayer)
 	return UPartyMember::StaticClass();
 }
 
-void USocialParty::HandlePartyStateChanged(const FUniqueNetId& LocalUserId, const FOnlinePartyId& PartyId, EPartyState PartyState)
+void USocialParty::HandlePartyStateChanged(const FUniqueNetId& LocalUserId, const FOnlinePartyId& PartyId, EPartyState PartyState, EPartyState PreviousPartyState)
 {
 	if (PartyState == EPartyState::Disconnected)
 	{
@@ -1137,7 +1143,7 @@ void USocialParty::HandlePartyStateChanged(const FUniqueNetId& LocalUserId, cons
 		// If we transition to the active state, then we have an XMPP connection
 		SetIsMissingXmppConnection(false);
 	}
-	OnPartyStateChanged().Broadcast(PartyState);
+	OnPartyStateChanged().Broadcast(PartyState, PreviousPartyState);
 }
 
 void USocialParty::ConnectToReservationBeacon()
diff --git a/Engine/Plugins/Online/OnlineFramework/Source/Party/Private/SocialManager.cpp b/Engine/Plugins/Online/OnlineFramework/Source/Party/Private/SocialManager.cpp
--- a/Engine/Plugins/Online/OnlineFramework/Source/Party/Private/SocialManager.cpp
+++ b/Engine/Plugins/Online/OnlineFramework/Source/Party/Private/SocialManager.cpp
@@ -715,7 +715,7 @@ USocialParty* USocialManager::EstablishNewParty(const FUniqueNetId& LocalUserId,
 		if (NewParty->IsPersistentParty())
 		{
 			NewParty->OnPartyStateChanged().AddUObject(this, &USocialManager::HandlePersistentPartyStateChanged, NewParty);
-			HandlePersistentPartyStateChanged(NewParty->GetOssPartyState(), NewParty);
+			HandlePersistentPartyStateChanged(NewParty->GetOssPartyState(), NewParty->GetOssPartyPreviousState(), NewParty);
 		}
 
 		return NewParty;
@@ -922,13 +922,18 @@ void USocialManager::HandleJoinPartyComplete(const FUniqueNetId& LocalUserId, co
 	}
 }
 
-void USocialManager::HandlePersistentPartyStateChanged(EPartyState NewState, USocialParty* PersistentParty)
+void USocialManager::HandlePersistentPartyStateChanged(EPartyState NewState, EPartyState PreviousState, USocialParty* PersistentParty)
 {
 	UE_LOG(LogParty, Verbose, TEXT("Persistent party state changed to %s"), ToString(NewState));
 
 	if (NewState == EPartyState::Disconnected)
 	{
 		bIsConnectedToPartyService = false;
+		
+		if (PreviousState == EPartyState::Active)
+		{
+			PersistentParty->LeaveParty();
+		}
 
 		// If we have other members in our party, then we will try to rejoin this when we come back online
 		if (!RejoinableParty.IsValid() && PersistentParty->ShouldCacheForRejoinOnDisconnect())
diff --git a/Engine/Plugins/Online/OnlineFramework/Source/Party/Public/Party/PartyMember.h b/Engine/Plugins/Online/OnlineFramework/Source/Party/Public/Party/PartyMember.h
--- a/Engine/Plugins/Online/OnlineFramework/Source/Party/Public/Party/PartyMember.h
+++ b/Engine/Plugins/Online/OnlineFramework/Source/Party/Public/Party/PartyMember.h
@@ -11,6 +11,7 @@ class USocialUser;
 class FOnlinePartyMember;
 class FOnlinePartyData;
 enum class EMemberExitedReason;
+enum class EMemberConnectionStatus;
 
 /** Base struct used to replicate data about the state of a single party member to all members. */
 USTRUCT()
@@ -80,6 +81,7 @@ public:
 	FUniqueNetIdRepl GetPrimaryNetId() const;
 	const FPartyMemberRepData& GetRepData() const { return *MemberDataReplicator; }
 	USocialUser& GetSocialUser() const;
+	EMemberConnectionStatus GetMemberConnectionStatus() const;
 
 	FString GetDisplayName() const;
 	FName GetPlatformOssName() const;
@@ -88,6 +90,7 @@ public:
 	FOnPartyMemberStateChanged& OnInitializationComplete() const { return OnMemberInitializedEvent; }
 	FOnPartyMemberStateChanged& OnPromotedToLeader() const { return OnPromotedToLeaderEvent; }
 	FOnPartyMemberStateChanged& OnDemoted() const { return OnDemotedEvent; }
+	FOnPartyMemberStateChanged& OnMemberConnectionStatusChanged() const { return OnMemberConnectionStatusChangedEvent; }
 
 	DECLARE_EVENT_OneParam(UPartyMember, FOnPartyMemberLeft, EMemberExitedReason)
 	FOnPartyMemberLeft& OnLeftParty() const { return OnLeftPartyEvent; }
@@ -113,9 +116,11 @@ protected:
 	virtual void Shutdown();
 
 	FPartyMemberDataReplicator MemberDataReplicator;
+	void SetMemberConnectionStatus(EMemberConnectionStatus InMemberConnectionStatus);
 
 private:
 	void HandleSocialUserInitialized(USocialUser& InitializedUser);
+	void HandleMemberConnectionStatusChanged(const FUniqueNetId& ChangedUserId, const EMemberConnectionStatus NewMemberConnectionStatus, const EMemberConnectionStatus PreviousMemberConnectionStatus);
 
 	TSharedPtr<FOnlinePartyMember> OssPartyMember;
 
@@ -123,6 +128,7 @@ private:
 	USocialUser* SocialUser = nullptr;
 
 	bool bHasReceivedInitialData = false;
+	mutable FOnPartyMemberStateChanged OnMemberConnectionStatusChangedEvent;
 	mutable FOnPartyMemberStateChanged OnMemberInitializedEvent;
 	mutable FOnPartyMemberStateChanged OnPromotedToLeaderEvent;
 	mutable FOnPartyMemberStateChanged OnDemotedEvent;
diff --git a/Engine/Plugins/Online/OnlineFramework/Source/Party/Public/Party/SocialParty.h b/Engine/Plugins/Online/OnlineFramework/Source/Party/Public/Party/SocialParty.h
--- a/Engine/Plugins/Online/OnlineFramework/Source/Party/Public/Party/SocialParty.h
+++ b/Engine/Plugins/Online/OnlineFramework/Source/Party/Public/Party/SocialParty.h
@@ -126,6 +126,7 @@ public:
 	const FOnlinePartyId& GetPartyId() const;
 
 	EPartyState GetOssPartyState() const;
+	EPartyState GetOssPartyPreviousState() const;
 
 	bool IsCurrentlyCrossplaying() const;
 	void StayWithPartyOnExit(bool bInStayWithParty);
@@ -167,7 +168,7 @@ public:
 	DECLARE_EVENT_OneParam(USocialParty, FOnPartyConfigurationChanged, const FPartyConfiguration&);
 	FOnPartyConfigurationChanged& OnPartyConfigurationChanged() const { return OnPartyConfigurationChangedEvent; }
 
-	DECLARE_EVENT_OneParam(USocialParty, FOnPartyStateChanged, EPartyState);
+	DECLARE_EVENT_TwoParams(USocialParty, FOnPartyStateChanged, EPartyState, EPartyState);
 	FOnPartyStateChanged& OnPartyStateChanged() const { return OnPartyStateChangedEvent; }
 
 	DECLARE_EVENT_OneParam(USocialParty, FOnPartyFunctionalityDegradedChanged, bool /*bFunctionalityDegraded*/);
@@ -179,6 +180,9 @@ public:
 	DECLARE_EVENT_TwoParams(USocialParty, FOnPartyJIPApproved, const FOnlinePartyId&, bool /* Success*/);
 	FOnPartyJIPApproved& OnPartyJIPApproved() const { return OnPartyJIPApprovedEvent; }
 
+	DECLARE_EVENT_TwoParams(USocialParty, FOnPartyMemberConnectionStatusChanged, UPartyMember&, EMemberConnectionStatus);
+	FOnPartyMemberConnectionStatusChanged& OnPartyMemberConnectionStatusChanged() const { return OnPartyMemberConnectionStatusChangedEvent; }
+
 	const FPartyPrivacySettings& GetPrivacySettings() const;
 
 PARTY_SCOPE:
@@ -282,7 +286,7 @@ private:
 	UPartyMember* GetMemberInternal(const FUniqueNetIdRepl& MemberId) const;
 
 private:	// Handlers
-	void HandlePartyStateChanged(const FUniqueNetId& LocalUserId, const FOnlinePartyId& PartyId, EPartyState PartyState);
+	void HandlePartyStateChanged(const FUniqueNetId& LocalUserId, const FOnlinePartyId& PartyId, EPartyState PartyState, EPartyState PreviousPartyState);
 	void HandlePartyConfigChanged(const FUniqueNetId& LocalUserId, const FOnlinePartyId& PartyId, const TSharedRef<FPartyConfiguration>& PartyConfig);
 	void HandleUpdatePartyConfigComplete(const FUniqueNetId& LocalUserId, const FOnlinePartyId& PartyId, EUpdateConfigCompletionResult Result);
 	void HandlePartyDataReceived(const FUniqueNetId& LocalUserId, const FOnlinePartyId& PartyId, const TSharedRef<FOnlinePartyData>& PartyData);
@@ -296,6 +300,7 @@ private:	// Handlers
 	void HandlePartyMemberJIP(const FUniqueNetId& LocalUserId, const FOnlinePartyId& PartyId, bool Success);
 	void HandlePartyMemberPromoted(const FUniqueNetId& LocalUserId, const FOnlinePartyId& PartyId, const FUniqueNetId& NewLeaderId);
 	void HandlePartyPromotionLockoutChanged(const FUniqueNetId& LocalUserId, const FOnlinePartyId& PartyId, bool bArePromotionsLocked);
+	void HandlePartyMemberConnectionStatusChanged(const FUniqueNetId& LocalUserId, const FOnlinePartyId& PartyId, EMemberConnectionStatus MemberConnectionStatus);
 
 	void HandleMemberInitialized(UPartyMember* Member);
 	void HandleMemberSessionIdChanged(const FSessionId& NewSessionId, UPartyMember* Member);
@@ -371,6 +376,7 @@ private:
 	mutable FOnPartyMemberCreated OnPartyMemberCreatedEvent;
 	mutable FOnPartyConfigurationChanged OnPartyConfigurationChangedEvent;
 	mutable FOnPartyStateChanged OnPartyStateChangedEvent;
+	mutable FOnPartyMemberConnectionStatusChanged OnPartyMemberConnectionStatusChangedEvent;
 	mutable FOnPartyFunctionalityDegradedChanged OnPartyFunctionalityDegradedChangedEvent;
 	mutable FOnInviteSent OnInviteSentEvent;
 	mutable FOnPartyJIPApproved OnPartyJIPApprovedEvent;
diff --git a/Engine/Plugins/Online/OnlineFramework/Source/Party/Public/SocialManager.h b/Engine/Plugins/Online/OnlineFramework/Source/Party/Public/SocialManager.h
--- a/Engine/Plugins/Online/OnlineFramework/Source/Party/Public/SocialManager.h
+++ b/Engine/Plugins/Online/OnlineFramework/Source/Party/Public/SocialManager.h
@@ -201,7 +201,7 @@ private:	// Handlers
 	void HandleCreatePartyComplete(const FUniqueNetId& LocalUserId, const TSharedPtr<const FOnlinePartyId>& PartyId, ECreatePartyCompletionResult Result, FOnlinePartyTypeId PartyTypeId, FOnCreatePartyAttemptComplete CompletionDelegate);
 	void HandleJoinPartyComplete(const FUniqueNetId& LocalUserId, const FOnlinePartyId& PartyId, EJoinPartyCompletionResult Result, int32 NotApprovedReasonCode, FOnlinePartyTypeId PartyTypeId);
 	
-	void HandlePersistentPartyStateChanged(EPartyState NewState, USocialParty* PersistentParty);
+	void HandlePersistentPartyStateChanged(EPartyState NewState, EPartyState PreviousState, USocialParty* PersistentParty);
 	void HandleLeavePartyForJoinComplete(ELeavePartyCompletionResult LeaveResult, USocialParty* LeftParty);
 	void HandlePartyLeaveBegin(EMemberExitedReason Reason, USocialParty* LeavingParty);
 	void HandlePartyLeft(EMemberExitedReason Reason, USocialParty* LeftParty);
