Additional fixes for fallout from adding the "all" parameter option for SOCIAL DEBUG commands.

#tests Compile DebugGame Editor Win64

diff --git a/Engine/Plugins/Online/OnlineFramework/Source/Party/Private/SocialDebugTools.cpp b/Engine/Plugins/Online/OnlineFramework/Source/Party/Private/SocialDebugTools.cpp
--- a/Engine/Plugins/Online/OnlineFramework/Source/Party/Private/SocialDebugTools.cpp
+++ b/Engine/Plugins/Online/OnlineFramework/Source/Party/Private/SocialDebugTools.cpp
@@ -78,6 +78,11 @@ bool USocialDebugTools::Exec(class UWorld* InWorld, const TCHAR* Cmd, FOutputDev
 					UE_LOG(LogParty, Log, TEXT("CONTEXT=ALL used, but no OSS contexts found!"));
 				}
 			}
+			else
+			{
+				TargetInstances.Add(Instance);
+			}
+
 			if (FParse::Command(&Cmd, TEXT("LOGIN")))
 			{
 				if (bAllInstancesRequested)
