Fixing incorrect parameter in LeaveParty SOCIAL DEBUG command

#tests Compile DebugGame Editor Win64

diff --git a/Engine/Plugins/Online/OnlineFramework/Source/Party/Private/SocialDebugTools.cpp b/Engine/Plugins/Online/OnlineFramework/Source/Party/Private/SocialDebugTools.cpp
--- a/Engine/Plugins/Online/OnlineFramework/Source/Party/Private/SocialDebugTools.cpp
+++ b/Engine/Plugins/Online/OnlineFramework/Source/Party/Private/SocialDebugTools.cpp
@@ -141,7 +141,7 @@ bool USocialDebugTools::Exec(class UWorld* InWorld, const TCHAR* Cmd, FOutputDev
 			{
 				for (const FString& TargetInstance : TargetInstances)
 				{
-					LeaveParty(Instance, FLeavePartyComplete::CreateLambda([this, TargetInstance](bool bSuccess)
+					LeaveParty(TargetInstance, FLeavePartyComplete::CreateLambda([this, TargetInstance](bool bSuccess)
 					{
 						UE_LOG(LogParty, Display, TEXT("Leave party OSS context[%s] %s"), *TargetInstance, *LexToString(bSuccess));
 					}));
