Added protected GetOwningMember accessor to SocialPartyMember for use by subclasses.

#tests Compile DebugGame Editor Win64

diff --git a/Engine/Plugins/Online/OnlineFramework/Source/Party/Private/Party/PartyMember.cpp b/Engine/Plugins/Online/OnlineFramework/Source/Party/Private/Party/PartyMember.cpp
--- a/Engine/Plugins/Online/OnlineFramework/Source/Party/Private/Party/PartyMember.cpp
+++ b/Engine/Plugins/Online/OnlineFramework/Source/Party/Private/Party/PartyMember.cpp
@@ -38,6 +38,11 @@ const USocialParty* FPartyMemberRepData::GetOwnerParty() const
 	return OwnerMember.IsValid() ? &OwnerMember->GetParty() : nullptr;
 }
 
+const UPartyMember* FPartyMemberRepData::GetOwningMember() const
+{
+	return OwnerMember.IsValid() ? OwnerMember.Get() : nullptr;
+}
+
 //////////////////////////////////////////////////////////////////////////
 // PartyMember
 //////////////////////////////////////////////////////////////////////////
diff --git a/Engine/Plugins/Online/OnlineFramework/Source/Party/Public/Party/PartyMember.h b/Engine/Plugins/Online/OnlineFramework/Source/Party/Public/Party/PartyMember.h
--- a/Engine/Plugins/Online/OnlineFramework/Source/Party/Public/Party/PartyMember.h
+++ b/Engine/Plugins/Online/OnlineFramework/Source/Party/Public/Party/PartyMember.h
@@ -27,6 +27,7 @@ protected:
 	virtual bool CanEditData() const override;
 	virtual void CompareAgainst(const FOnlinePartyRepDataBase& OldData) const override;
 	virtual const USocialParty* GetOwnerParty() const override;
+	virtual const UPartyMember* GetOwningMember() const;
 
 private:
 	TWeakObjectPtr<const UPartyMember> OwnerMember;
