Add utilities to more easily get subclassed Toolkits and PartyMembers in Social.

diff --git a/Engine/Plugins/Online/OnlineFramework/Source/Party/Private/Party/SocialParty.cpp b/Engine/Plugins/Online/OnlineFramework/Source/Party/Private/Party/SocialParty.cpp
--- a/Engine/Plugins/Online/OnlineFramework/Source/Party/Private/Party/SocialParty.cpp
+++ b/Engine/Plugins/Online/OnlineFramework/Source/Party/Private/Party/SocialParty.cpp
@@ -416,7 +416,7 @@ void USocialParty::InitializePartyInternal()
 	PartyInterface->AddOnPartyMemberDataReceivedDelegate_Handle(FOnPartyMemberDataReceivedDelegate::CreateUObject(this, &USocialParty::HandlePartyMemberDataReceived));
 	PartyInterface->AddOnPartyMemberPromotedDelegate_Handle(FOnPartyMemberPromotedDelegate::CreateUObject(this, &USocialParty::HandlePartyMemberPromoted));
 	PartyInterface->AddOnPartyMemberExitedDelegate_Handle(FOnPartyMemberExitedDelegate::CreateUObject(this, &USocialParty::HandlePartyMemberExited));
-	//PartyInterface->AddOnPartyMemberStateChangedDelegate_Handle(FOnPartyMemberStateChangedDelegate::CreateUObject(this, &USocialParty::HandlePartyMemberStateChanged));
+	
 	// Create a UPartyMember for every existing member on the OSS party
 	TArray<TSharedRef<FOnlinePartyMember>> OssPartyMembers;
 	PartyInterface->GetPartyMembers(*OwningLocalUserId, GetPartyId(), OssPartyMembers);
@@ -535,9 +535,7 @@ UPartyMember* USocialParty::GetOrCreatePartyMember(const FUniqueNetId& MemberId)
 					PartyMembersById.Add(MemberIdRepl, PartyMember);
 					PartyMember->InitializePartyMember(OssPartyMember.ToSharedRef(), FSimpleDelegate::CreateUObject(this, &USocialParty::HandleMemberInitialized, PartyMember));
 
-					RefreshPublicJoinability();
-
-					OnPartyMemberCreated().Broadcast(*PartyMember);
+					OnMemberCreatedInternal(*PartyMember);
 				}
 				else
 				{
@@ -864,6 +862,12 @@ void USocialParty::HandlePrivacySettingsChanged(const FPartyPrivacySettings& New
 	RefreshPublicJoinability();
 }
 
+void USocialParty::OnMemberCreatedInternal(UPartyMember& NewMember)
+{
+	RefreshPublicJoinability();
+	OnPartyMemberCreated().Broadcast(NewMember);
+}
+
 void USocialParty::HandlePartyLeft(const FUniqueNetId& LocalUserId, const FOnlinePartyId& PartyId)
 {
 	// this function is called when a party is left due to unintentional leave (e.g. disconnect)
diff --git a/Engine/Plugins/Online/OnlineFramework/Source/Party/Private/SocialToolkit.cpp b/Engine/Plugins/Online/OnlineFramework/Source/Party/Private/SocialToolkit.cpp
--- a/Engine/Plugins/Online/OnlineFramework/Source/Party/Private/SocialToolkit.cpp
+++ b/Engine/Plugins/Online/OnlineFramework/Source/Party/Private/SocialToolkit.cpp
@@ -110,7 +110,7 @@ public:
 
 //@todo DanH Social: Need a non-backdoor way to get toolkits from the manager (an issue when we don't know where the manager is) - new game subsystems should be a nice solve
 TMap<TWeakObjectPtr<ULocalPlayer>, TWeakObjectPtr<USocialToolkit>> USocialToolkit::AllToolkitsByOwningPlayer;
-USocialToolkit* USocialToolkit::GetToolkitForPlayer(ULocalPlayer* LocalPlayer)
+USocialToolkit* USocialToolkit::GetToolkitForPlayerInternal(ULocalPlayer* LocalPlayer)
 {
 	TWeakObjectPtr<USocialToolkit>* FoundToolkit = AllToolkitsByOwningPlayer.Find(LocalPlayer);
 	return FoundToolkit ? FoundToolkit->Get() : nullptr;
diff --git a/Engine/Plugins/Online/OnlineFramework/Source/Party/Public/Party/PartyTypes.h b/Engine/Plugins/Online/OnlineFramework/Source/Party/Public/Party/PartyTypes.h
--- a/Engine/Plugins/Online/OnlineFramework/Source/Party/Public/Party/PartyTypes.h
+++ b/Engine/Plugins/Online/OnlineFramework/Source/Party/Public/Party/PartyTypes.h
@@ -437,15 +437,9 @@ public:	\
 	FOn##Property##ChangedDif& On##Property##ChangedDif() const { return On##Property##ChangedDifEvent; }	\
 	\
 	Property##ArgType Get##Property() const { return Property; }	\
-	\
-	bool Has##Property##InitiallyReplicated() const { return b##Property##InitiallyReplicated; }	\
 private:	\
 	void Compare##Property(const Owner& OldData) const	\
 	{	\
-		if(!b##Property##InitiallyReplicated) \
-		{ \
-			b##Property##InitiallyReplicated = true; \
-		} \
 		if (Property != OldData.Property)	\
 		{	\
 			LogPropertyChanged(TEXT(#Owner), TEXT(#Property), true);	\
@@ -454,8 +448,7 @@ private:	\
 		}	\
 	}	\
 	mutable FOn##Property##Changed On##Property##ChangedEvent;	\
-	mutable FOn##Property##ChangedDif On##Property##ChangedDifEvent; \
-	mutable bool b##Property##InitiallyReplicated
+	mutable FOn##Property##ChangedDif On##Property##ChangedDifEvent
 
 /**
  * Exposes a rep data property and provides a default property setter.
diff --git a/Engine/Plugins/Online/OnlineFramework/Source/Party/Public/Party/SocialParty.h b/Engine/Plugins/Online/OnlineFramework/Source/Party/Public/Party/SocialParty.h
--- a/Engine/Plugins/Online/OnlineFramework/Source/Party/Public/Party/SocialParty.h
+++ b/Engine/Plugins/Online/OnlineFramework/Source/Party/Public/Party/SocialParty.h
@@ -218,6 +218,7 @@ protected:
 	virtual FPartyPrivacySettings GetDesiredPrivacySettings() const;
 	virtual void OnLocalPlayerIsLeaderChanged(bool bIsLeader);
 	virtual void HandlePrivacySettingsChanged(const FPartyPrivacySettings& NewPrivacySettings);
+	virtual void OnMemberCreatedInternal(UPartyMember& NewMember);
 	virtual void OnLeftPartyInternal(EMemberExitedReason Reason);
 
 	virtual void OnInviteSentInternal(ESocialSubsystem SubsystemType, const USocialUser& InvitedUser, bool bWasSuccessful);
@@ -300,7 +301,6 @@ private:	// Handlers
 	void HandlePartyMemberJIP(const FUniqueNetId& LocalUserId, const FOnlinePartyId& PartyId, bool Success);
 	void HandlePartyMemberPromoted(const FUniqueNetId& LocalUserId, const FOnlinePartyId& PartyId, const FUniqueNetId& NewLeaderId);
 	void HandlePartyPromotionLockoutChanged(const FUniqueNetId& LocalUserId, const FOnlinePartyId& PartyId, bool bArePromotionsLocked);
-	void HandlePartyMemberConnectionStatusChanged(const FUniqueNetId& LocalUserId, const FOnlinePartyId& PartyId, EMemberConnectionStatus MemberConnectionStatus);
 
 	void HandleMemberInitialized(UPartyMember* Member);
 	void HandleMemberSessionIdChanged(const FSessionId& NewSessionId, UPartyMember* Member);
@@ -348,7 +348,7 @@ private:
 	
 	/** Reservation beacon client instance while getting approval for new party members*/
 	UPROPERTY()
-		APartyBeaconClient* ReservationBeaconClient = nullptr;
+	APartyBeaconClient* ReservationBeaconClient = nullptr;
 	
 	/**
 	* Last known spectator beacon client net driver name
@@ -359,7 +359,7 @@ private:
 	
 	/** Spectator beacon client instance while getting approval for spectator*/
 	UPROPERTY()
-		ASpectatorBeaconClient* SpectatorBeaconClient = nullptr;
+	ASpectatorBeaconClient* SpectatorBeaconClient = nullptr;
 
 	/**
 	 * True when we have limited functionality due to lacking an xmpp connection.
diff --git a/Engine/Plugins/Online/OnlineFramework/Source/Party/Public/SocialToolkit.h b/Engine/Plugins/Online/OnlineFramework/Source/Party/Public/SocialToolkit.h
--- a/Engine/Plugins/Online/OnlineFramework/Source/Party/Public/SocialToolkit.h
+++ b/Engine/Plugins/Online/OnlineFramework/Source/Party/Public/SocialToolkit.h
@@ -30,7 +30,12 @@ class PARTY_API USocialToolkit : public UObject
 	GENERATED_BODY()
 
 public:
-	static USocialToolkit* GetToolkitForPlayer(ULocalPlayer* LocalPlayer);
+	template <typename ToolkitT = USocialToolkit>
+	static ToolkitT* GetToolkitForPlayer(ULocalPlayer* LocalPlayer)
+	{
+		static_assert(TIsDerivedFrom<ToolkitT, USocialToolkit>::IsDerived, "GetToolkitForPlayer only supports getting USocialToolkit type objects");
+		return Cast<ToolkitT>(GetToolkitForPlayerInternal(LocalPlayer));
+	}
 
 	USocialToolkit();
 
@@ -197,6 +202,7 @@ private:	// Handlers
 	void HandleGameDestroyed(const FName SessionName, bool bWasSuccessful);
 
 private:
+	static USocialToolkit* GetToolkitForPlayerInternal(ULocalPlayer* LocalPlayer);
 	static TMap<TWeakObjectPtr<ULocalPlayer>, TWeakObjectPtr<USocialToolkit>> AllToolkitsByOwningPlayer;
 
 	UPROPERTY()
