Always initialize the local player's party member first

diff --git a/Engine/Plugins/Online/OnlineFramework/Source/Party/Private/Party/SocialParty.cpp b/Engine/Plugins/Online/OnlineFramework/Source/Party/Private/Party/SocialParty.cpp
--- a/Engine/Plugins/Online/OnlineFramework/Source/Party/Private/Party/SocialParty.cpp
+++ b/Engine/Plugins/Online/OnlineFramework/Source/Party/Private/Party/SocialParty.cpp
@@ -420,6 +420,11 @@ void USocialParty::InitializePartyInternal()
 	// Create a UPartyMember for every existing member on the OSS party
 	TArray<TSharedRef<FOnlinePartyMember>> OssPartyMembers;
 	PartyInterface->GetPartyMembers(*OwningLocalUserId, GetPartyId(), OssPartyMembers);
+	// Always initialize the local member first
+	if (ensure(OssPartyMembers.RemoveAll([this](const TSharedRef<FOnlinePartyMember>& Member) { return *Member->GetUserId() == *OwningLocalUserId; } ) > 0))
+	{
+		GetOrCreatePartyMember(*OwningLocalUserId);
+	}
 	for (TSharedRef<FOnlinePartyMember>& OssMember : OssPartyMembers)
 	{
 		GetOrCreatePartyMember(*OssMember->GetUserId());
