Expose IOnlinePartySystem::RestoreParties in the social manager

diff --git a/Engine/Plugins/Online/OnlineFramework/Source/Party/Private/SocialManager.cpp b/Engine/Plugins/Online/OnlineFramework/Source/Party/Private/SocialManager.cpp
--- a/Engine/Plugins/Online/OnlineFramework/Source/Party/Private/SocialManager.cpp
+++ b/Engine/Plugins/Online/OnlineFramework/Source/Party/Private/SocialManager.cpp
@@ -846,6 +846,47 @@ void USocialManager::HandleLocalPlayerRemoved(int32 LocalUserNum)
 	}
 }
 
+void USocialManager::RestorePartyStateFromPartySystem(const FOnRestorePartyStateFromPartySystemComplete& OnRestoreComplete)
+{
+	UE_LOG(LogParty, Verbose, TEXT("RestorePartyStateFromPartySystem"));
+	FUniqueNetIdRepl LocalUserId = GetFirstLocalUserId(ESocialSubsystem::Primary);
+
+	// If the player has any parties, do not try to restore
+	if (LocalUserId.IsValid() &&
+		bCanCreatePartyObjects &&
+		JoinedPartiesByTypeId.Num() == 0 &&
+		JoinAttemptsByTypeId.Num() == 0)
+	{
+		IOnlinePartyPtr PartyInterface = Online::GetPartyInterfaceChecked(GetWorld());
+		PartyInterface->RestoreParties(*LocalUserId, FOnRestorePartiesComplete::CreateUObject(this, &USocialManager::OnRestorePartiesComplete, OnRestoreComplete));
+	}
+	else
+	{
+		OnRestoreComplete.ExecuteIfBound(false);
+	}
+}
+
+void USocialManager::OnRestorePartiesComplete(const FUniqueNetId& LocalUserId, const FOnlineError& Result, const FOnRestorePartyStateFromPartySystemComplete OnRestoreComplete)
+{
+	if (Result.WasSuccessful())
+	{
+		// Restore our parties
+		IOnlinePartyPtr PartyInterface = Online::GetPartyInterfaceChecked(GetWorld());
+		TArray<TSharedRef<const FOnlinePartyId>> JoinedParties;
+		PartyInterface->GetJoinedParties(LocalUserId, JoinedParties);
+		for (const TSharedRef<const FOnlinePartyId>& PartyId : JoinedParties)
+		{
+			TSharedPtr<const FOnlineParty> Party = PartyInterface->GetParty(LocalUserId, *PartyId);
+			check(Party.IsValid());
+			if (!EstablishNewParty(LocalUserId, *PartyId, Party->PartyTypeId))
+			{
+				UE_LOG(LogParty, Warning, TEXT("OnRestorePartiesComplete: User=[%s] Party=[%s] Type=%d failed to establish party"), *LocalUserId.ToDebugString(), *PartyId->ToDebugString(), Party->PartyTypeId.GetValue());
+			}
+		}
+	}
+	OnRestoreComplete.ExecuteIfBound(Result.WasSuccessful());
+}
+
 void USocialManager::HandleQueryJoinabilityComplete(const FUniqueNetId& LocalUserId, const FOnlinePartyId& PartyId, EJoinPartyCompletionResult Result, int32 NotApprovedReasonCode, FOnlinePartyTypeId PartyTypeId)
 {
 	if (FJoinPartyAttempt* JoinAttempt = JoinAttemptsByTypeId.Find(PartyTypeId))
diff --git a/Engine/Plugins/Online/OnlineFramework/Source/Party/Public/SocialManager.h b/Engine/Plugins/Online/OnlineFramework/Source/Party/Public/SocialManager.h
--- a/Engine/Plugins/Online/OnlineFramework/Source/Party/Public/SocialManager.h
+++ b/Engine/Plugins/Online/OnlineFramework/Source/Party/Public/SocialManager.h
@@ -65,6 +65,10 @@ public:
 	void CreateParty(const FOnlinePartyTypeId& PartyTypeId, const FPartyConfiguration& PartyConfig, const FOnCreatePartyAttemptComplete& OnCreatePartyComplete);
 	void CreatePersistentParty(const FOnCreatePartyAttemptComplete& OnCreatePartyComplete = FOnCreatePartyAttemptComplete());
 
+	/** Attempt to restore our party state from the party system */
+	DECLARE_DELEGATE_OneParam(FOnRestorePartyStateFromPartySystemComplete, bool /*bSucceeded*/)
+	void RestorePartyStateFromPartySystem(const FOnRestorePartyStateFromPartySystemComplete& OnRestoreComplete);
+
 	bool IsPartyJoinInProgress(const FOnlinePartyTypeId& TypeId) const;
 	bool IsPersistentPartyJoinInProgress() const;
 
@@ -203,7 +207,8 @@ private:	// Handlers
 	void HandleWorldEstablished(UWorld* World);
 	void HandleLocalPlayerAdded(int32 LocalUserNum);
 	void HandleLocalPlayerRemoved(int32 LocalUserNum);
-
+	
+	void OnRestorePartiesComplete(const FUniqueNetId& LocalUserId, const FOnlineError& Result, const FOnRestorePartyStateFromPartySystemComplete OnRestoreComplete);
 	void HandleQueryJoinabilityComplete(const FUniqueNetId& LocalUserId, const FOnlinePartyId& PartyId, EJoinPartyCompletionResult Result, int32 NotApprovedReasonCode, FOnlinePartyTypeId PartyTypeId);
 	void HandleCreatePartyComplete(const FUniqueNetId& LocalUserId, const TSharedPtr<const FOnlinePartyId>& PartyId, ECreatePartyCompletionResult Result, FOnlinePartyTypeId PartyTypeId, FOnCreatePartyAttemptComplete CompletionDelegate);
 	void HandleJoinPartyComplete(const FUniqueNetId& LocalUserId, const FOnlinePartyId& PartyId, EJoinPartyCompletionResult Result, int32 NotApprovedReasonCode, FOnlinePartyTypeId PartyTypeId);
