Changed SocialUserList::AutoUpdateRate to be globally hotfixable, more optimizations are needed and being worked on.

#tests Compile DebugGame Editor Win64 / Verified hotfixability in PIE

diff --git a/Engine/Plugins/Online/OnlineFramework/Source/Party/Private/SocialSettings.cpp b/Engine/Plugins/Online/OnlineFramework/Source/Party/Private/SocialSettings.cpp
--- a/Engine/Plugins/Online/OnlineFramework/Source/Party/Private/SocialSettings.cpp
+++ b/Engine/Plugins/Online/OnlineFramework/Source/Party/Private/SocialSettings.cpp
@@ -33,3 +33,9 @@ int32 USocialSettings::GetDefaultMaxPartySize()
 	const USocialSettings& SettingsCDO = *GetDefault<USocialSettings>();
 	return SettingsCDO.DefaultMaxPartySize;
 }
+
+float USocialSettings::GetUserListAutoUpdateRate()
+{
+	const USocialSettings& SettingsCDO = *GetDefault<USocialSettings>();
+	return SettingsCDO.UserListAutoUpdateRate;
+}
diff --git a/Engine/Plugins/Online/OnlineFramework/Source/Party/Private/User/SocialUserList.cpp b/Engine/Plugins/Online/OnlineFramework/Source/Party/Private/User/SocialUserList.cpp
--- a/Engine/Plugins/Online/OnlineFramework/Source/Party/Private/User/SocialUserList.cpp
+++ b/Engine/Plugins/Online/OnlineFramework/Source/Party/Private/User/SocialUserList.cpp
@@ -7,6 +7,7 @@
 
 #include "Containers/Ticker.h"
 #include "Interfaces/OnlinePresenceInterface.h"
+#include "SocialSettings.h"
 
 TSharedRef<FSocialUserList> FSocialUserList::CreateUserList(USocialToolkit& InOwnerToolkit, const FSocialUserListConfig& InConfig)
 {
@@ -61,7 +62,7 @@ void FSocialUserList::InitializeList()
 		TryAddUserFast(*User);
 	}
 
-	SetAutoUpdatePeriod(AutoUpdatePeriod);
+	SetAutoUpdatePeriod(USocialSettings::GetUserListAutoUpdateRate());
 }
 
 void FSocialUserList::AddReferencedObjects(FReferenceCollector& Collector)
diff --git a/Engine/Plugins/Online/OnlineFramework/Source/Party/Public/SocialSettings.h b/Engine/Plugins/Online/OnlineFramework/Source/Party/Public/SocialSettings.h
--- a/Engine/Plugins/Online/OnlineFramework/Source/Party/Public/SocialSettings.h
+++ b/Engine/Plugins/Online/OnlineFramework/Source/Party/Public/SocialSettings.h
@@ -23,6 +23,7 @@ public:
 	static FString GetUniqueIdEnvironmentPrefix(ESocialSubsystem SubsystemType);
 	static bool ShouldPreferPlatformInvites();
 	static int32 GetDefaultMaxPartySize();
+	static float GetUserListAutoUpdateRate();
 
 private:
 	/**
@@ -40,4 +41,7 @@ private:
 
 	UPROPERTY(config)
 	bool bPreferPlatformInvites = true;
+
+	UPROPERTY(config)
+	float UserListAutoUpdateRate = 7.5f;
 };
\ No newline at end of file
