Unflag missing XMPP when a party state changes from CreatePending to Active
When a party goes from Active to Disconnected, we leave the party and mark us as experiencing difficulties, and create a new party.  If the new party goes from CreatePending to Active, we don't clear the missing XMPP flag

diff --git a/Engine/Plugins/Online/OnlineFramework/Source/Party/Private/Party/SocialParty.cpp b/Engine/Plugins/Online/OnlineFramework/Source/Party/Private/Party/SocialParty.cpp
--- a/Engine/Plugins/Online/OnlineFramework/Source/Party/Private/Party/SocialParty.cpp
+++ b/Engine/Plugins/Online/OnlineFramework/Source/Party/Private/Party/SocialParty.cpp
@@ -987,7 +987,7 @@ bool USocialParty::ShouldStayWithPartyOnExit() const
 
 bool USocialParty::IsPartyFunctionalityDegraded() const
 {
-	return bIsMissingXmppConnection || bIsMissingPlatformSession;
+	return bIsMissingXmppConnection.Get(false) || bIsMissingPlatformSession;
 }
 
 int32 USocialParty::GetNumPartyMembers() const
@@ -1413,9 +1413,10 @@ void USocialParty::SetIsMissingPlatformSession(bool bInIsMissingPlatformSession)
 
 void USocialParty::SetIsMissingXmppConnection(bool bInMissingXmppConnection)
 {
-	if (bInMissingXmppConnection != bIsMissingXmppConnection)
+	if (!bIsMissingXmppConnection.IsSet() || 
+		bInMissingXmppConnection != bIsMissingXmppConnection)
 	{
-		UE_LOG(LogParty, VeryVerbose, TEXT("Party [%s] is %s missing XMPP connection"), *ToDebugString(), bInMissingXmppConnection ? TEXT("now") : TEXT("no longer"));
+		UE_CLOG(bIsMissingXmppConnection.IsSet(), LogParty, VeryVerbose, TEXT("Party [%s] is %s missing XMPP connection"), *ToDebugString(), bInMissingXmppConnection ? TEXT("now") : TEXT("no longer"));
 
 		const bool bWasPartyFunctionalityDegraded = IsPartyFunctionalityDegraded();
 		bIsMissingXmppConnection = bInMissingXmppConnection;
diff --git a/Engine/Plugins/Online/OnlineFramework/Source/Party/Public/Party/SocialParty.h b/Engine/Plugins/Online/OnlineFramework/Source/Party/Public/Party/SocialParty.h
--- a/Engine/Plugins/Online/OnlineFramework/Source/Party/Public/Party/SocialParty.h
+++ b/Engine/Plugins/Online/OnlineFramework/Source/Party/Public/Party/SocialParty.h
@@ -365,7 +365,7 @@ private:
 	 * True when we have limited functionality due to lacking an xmpp connection.
 	 * Don't set directly, use the private setter to trigger events appropriately.
 	 */
-	bool bIsMissingXmppConnection = false;
+	TOptional<bool> bIsMissingXmppConnection;
 	bool bIsMissingPlatformSession = false;
 
 	bool bIsLeavingParty = false;
