Optimized FSocialUserList::HandleAutoUpdateList a bit:
* Remove items from PendingRemovals as we identify them in the UserList to save future iterations in PendingRemovals.Contains()
* Removed "Is Party Member" comparison from Sort as it might not be universally desired or relevant.
* Cache all the per-User data we need for sorting in an FUserSortData struct array, this array is then sorted and the contents of Users are replaced from the sorted struct array after the sort is completed.

#tests Compile DebugGame Editor Win64 / PIE verification of proper sorting

diff --git a/Engine/Plugins/Online/OnlineFramework/Source/Party/Private/User/SocialUserList.cpp b/Engine/Plugins/Online/OnlineFramework/Source/Party/Private/User/SocialUserList.cpp
--- a/Engine/Plugins/Online/OnlineFramework/Source/Party/Private/User/SocialUserList.cpp
+++ b/Engine/Plugins/Online/OnlineFramework/Source/Party/Private/User/SocialUserList.cpp
@@ -8,6 +8,7 @@
 #include "Containers/Ticker.h"
 #include "Interfaces/OnlinePresenceInterface.h"
 #include "SocialSettings.h"
+#include "Algo/Transform.h"
 
 TSharedRef<FSocialUserList> FSocialUserList::CreateUserList(USocialToolkit& InOwnerToolkit, const FSocialUserListConfig& InConfig)
 {
@@ -357,6 +358,40 @@ bool FSocialUserList::EvaluatePresenceFlag(bool bPresenceValue, ESocialUserState
 	return true;
 }
 
+// encapsulates UserList sorting comparator and supporting data needed
+struct FUserSortData
+{
+	FUserSortData(USocialUser* InUser, EOnlinePresenceState::Type InStatus, bool InPlayingThisGame, FString InDisplayName)
+		: User(InUser), OnlineStatus(InStatus), PlayingThisGame(InPlayingThisGame), DisplayName(MoveTemp(InDisplayName))
+	{ }
+
+	USocialUser* User;
+	EOnlinePresenceState::Type OnlineStatus;
+	bool PlayingThisGame;
+	FString DisplayName;
+
+	bool operator<(const FUserSortData& OtherSortData) const
+	{
+		// Goes from if online, then alphabetical
+		if (OnlineStatus == OtherSortData.OnlineStatus)
+		{
+			if (PlayingThisGame == OtherSortData.PlayingThisGame)
+			{
+				return DisplayName < OtherSortData.DisplayName;
+			}
+			else
+			{
+				return false;
+			}
+		}
+		else
+		{
+			// @todo StephanJ: note Online < Offline < Away, but it's okay for now since we show offline in a separate list #future
+			return OnlineStatus < OtherSortData.OnlineStatus;
+		}
+	}
+};
+
 bool FSocialUserList::HandleAutoUpdateList(float)
 {
 	// Re-evaluate whether each user with dirtied presence is still fit for the list
@@ -386,16 +421,17 @@ bool FSocialUserList::HandleAutoUpdateList(float)
 		Users.RemoveAllSwap(
 			[this] (USocialUser* User)
 			{
-				return PendingRemovals.Contains(User);
+				if (PendingRemovals.Contains(User))
+				{
+					PendingRemovals.Remove(User);
+					OnUserRemoved().Broadcast(*User);
+					return true;
+				}
+
+				return false;
 			});
 
-		for (TWeakObjectPtr<const USocialUser> User : PendingRemovals)
-		{
-			if (User.IsValid())
-			{
-				OnUserRemoved().Broadcast(*User.Get());
-			}
-		}
+		ensure(PendingRemovals.Num() == 0);
 		PendingRemovals.Reset();
 	}
 	
@@ -415,41 +451,23 @@ bool FSocialUserList::HandleAutoUpdateList(float)
 	{
 		bNeedsSort = false;
 
-		Users.Sort([](USocialUser& UserA, USocialUser& UserB)
-			{
-				const UPartyMember* PartyMemberA = UserA.GetPartyMember(IOnlinePartySystem::GetPrimaryPartyTypeId());
-				const UPartyMember* PartyMemberB = UserB.GetPartyMember(IOnlinePartySystem::GetPrimaryPartyTypeId());
+		const int32 NumUsers = Users.Num();
+		TArray<FUserSortData> SortedData;
+		SortedData.Reserve(NumUsers);
 
-				// Put party members at the top
-				if (PartyMemberA && !PartyMemberB)
-				{
-					return true;
-				}
-				else if (PartyMemberB && !PartyMemberA)
-				{
-					return false;
-				}
+		Algo::Transform(Users, SortedData, [](USocialUser* const User) -> FUserSortData
+		{
+			return FUserSortData(User, User->GetOnlineStatus(), User->IsPlayingThisGame(), User->GetDisplayName());
+		});
 
-				// Goes from if online, then alphabetical
-				if (UserA.GetOnlineStatus() == UserB.GetOnlineStatus())
-				{
-					if (UserA.IsPlayingThisGame() == UserB.IsPlayingThisGame())
-					{
-						return UserA.GetDisplayName() < UserB.GetDisplayName();
-					}
-					else
-					{
-						return UserA.IsPlayingThisGame() > UserB.IsPlayingThisGame();
-					}
-				}
-				else
-				{
-					// @todo StephanJ: note Online < Offline < Away, but it's okay for now since we show offline in a separate list #future
-					return UserA.GetOnlineStatus() < UserB.GetOnlineStatus();
-				}
-				
-			});
+		Algo::Sort(SortedData);
 
+		// replace contents of Users from SortedData array
+		for (int Index = 0; Index < NumUsers; Index++)
+		{
+			Users[Index] = SortedData[Index].User;
+		}
+		
 		OnUpdateComplete().Broadcast();
 	}
 
