-- Default all list config to not require auto update, and skip the ticker for most use cases.
  -- On top of that, added a bAllowAutoUpdate (default to true) for external overwrite to disable auto update when the list is not visible (or useful) to the users. even bRequireAutoUpdate is true, this maintains the ticker but skip the handler function.

diff --git a/Engine/Plugins/Online/OnlineFramework/Source/Party/Private/User/SocialUserList.cpp b/Engine/Plugins/Online/OnlineFramework/Source/Party/Private/User/SocialUserList.cpp
--- a/Engine/Plugins/Online/OnlineFramework/Source/Party/Private/User/SocialUserList.cpp
+++ b/Engine/Plugins/Online/OnlineFramework/Source/Party/Private/User/SocialUserList.cpp
@@ -63,7 +63,10 @@ void FSocialUserList::InitializeList()
 		TryAddUserFast(*User);
 	}
 
-	SetAutoUpdatePeriod(USocialSettings::GetUserListAutoUpdateRate());
+	if (ListConfig.bRequireAutoUpdate)
+	{
+		SetAutoUpdatePeriod(USocialSettings::GetUserListAutoUpdateRate());
+	}
 }
 
 void FSocialUserList::AddReferencedObjects(FReferenceCollector& Collector)
@@ -73,7 +76,7 @@ void FSocialUserList::AddReferencedObjects(FReferenceCollector& Collector)
 
 void FSocialUserList::UpdateNow()
 {
-	HandleAutoUpdateList(0.f);
+	UpdateListInternal();
 }
 
 void FSocialUserList::SetAutoUpdatePeriod(float InAutoUpdatePeriod)
@@ -119,23 +122,25 @@ void FSocialUserList::HandleOwnerToolkitReset()
 void FSocialUserList::HandlePartyInviteReceived(USocialUser& InvitingUser)
 {
 	TryAddUser(InvitingUser);
+	UpdateListInternal();
 }
 
 void FSocialUserList::HandlePartyInviteHandled(USocialUser* InvitingUser)
 {
 	TryRemoveUser(*InvitingUser);
-	UpdateNow();
+	UpdateListInternal();
 }
 
 void FSocialUserList::HandleFriendInviteReceived(USocialUser& User, ESocialSubsystem SubsystemType)
 {
 	TryAddUser(User);
+	UpdateListInternal();
 }
 
 void FSocialUserList::HandleFriendInviteRemoved(ESocialSubsystem SubsystemType, USocialUser* User)
 {
 	TryRemoveUser(*User);
-	UpdateNow();
+	UpdateListInternal();
 }
 
 void FSocialUserList::HandleFriendshipEstablished(USocialUser& NewFriend, ESocialSubsystem SubsystemType, bool bIsNewRelationship)
@@ -151,14 +156,14 @@ void FSocialUserList::HandleFriendshipEstablished(USocialUser& NewFriend, ESocia
 	{
 		// Any non-friends list that cares about friendship does so to remove entries (i.e. invites & recent players)
 		TryRemoveUser(NewFriend);
-		UpdateNow();
+		UpdateListInternal();
 	}
 }
 
 void FSocialUserList::HandleFriendRemoved(ESocialSubsystem SubsystemType, USocialUser* User)
 {
 	TryRemoveUser(*User);
-	UpdateNow();
+	UpdateListInternal();
 }
 
 void FSocialUserList::HandleUserBlocked(USocialUser& BlockedUser, ESocialSubsystem SubsystemType, bool bIsNewRelationship)
@@ -173,7 +178,7 @@ void FSocialUserList::HandleUserBlocked(USocialUser& BlockedUser, ESocialSubsyst
 		TryRemoveUser(BlockedUser);
 	}
 	
-	UpdateNow();
+	UpdateListInternal();
 }
 
 void FSocialUserList::HandleUserBlockStatusChanged(ESocialSubsystem SubsystemType, bool bIsBlocked, USocialUser* User)
@@ -181,7 +186,7 @@ void FSocialUserList::HandleUserBlockStatusChanged(ESocialSubsystem SubsystemTyp
 	if (!bIsBlocked)
 	{
 		TryRemoveUser(*User);
-		UpdateNow();
+		UpdateListInternal();
 	}
 }
 
@@ -393,6 +398,15 @@ struct FUserSortData
 };
 
 bool FSocialUserList::HandleAutoUpdateList(float)
+{
+	if (bAllowAutoUpdate)
+	{
+		UpdateListInternal();
+	}
+	return true;
+}
+
+void FSocialUserList::UpdateListInternal()
 {
 	// Re-evaluate whether each user with dirtied presence is still fit for the list
 	for (TWeakObjectPtr<USocialUser> DirtyUser : UsersWithDirtyPresence)
@@ -434,12 +448,12 @@ bool FSocialUserList::HandleAutoUpdateList(float)
 		ensure(PendingRemovals.Num() == 0);
 		PendingRemovals.Reset();
 	}
-	
+
 	if (PendingAdds.Num() > 0)
 	{
 		bListChanged = true;
 		Users.Append(PendingAdds);
-		
+
 		for (USocialUser* User : PendingAdds)
 		{
 			OnUserAdded().Broadcast(*User);
@@ -470,6 +484,4 @@ bool FSocialUserList::HandleAutoUpdateList(float)
 		
 		OnUpdateComplete().Broadcast();
 	}
-
-	return true;
 }
diff --git a/Engine/Plugins/Online/OnlineFramework/Source/Party/Private/User/SocialUserList.h b/Engine/Plugins/Online/OnlineFramework/Source/Party/Private/User/SocialUserList.h
--- a/Engine/Plugins/Online/OnlineFramework/Source/Party/Private/User/SocialUserList.h
+++ b/Engine/Plugins/Online/OnlineFramework/Source/Party/Private/User/SocialUserList.h
@@ -17,6 +17,7 @@ public:
 	FOnUpdateComplete& OnUpdateComplete() const override { return OnUpdateCompleteEvent; }
 
 	void UpdateNow();
+	void SetAllowAutoUpdate(bool bIsEnabled) { bAllowAutoUpdate = bIsEnabled; }
 	void SetAutoUpdatePeriod(float InAutoUpdatePeriod);
 	const TArray<USocialUser*>& GetUsers() const { return Users; }
 
@@ -52,6 +53,7 @@ private:
 	bool EvaluatePresenceFlag(bool bPresenceValue, ESocialUserStateFlags Flag) const;
 
 	bool HandleAutoUpdateList(float);
+	void UpdateListInternal();
 
 private:
 	FSocialUserList(USocialToolkit& InOwnerToolkit, const FSocialUserListConfig& Config);
@@ -70,6 +72,8 @@ private:
 
 	FSocialUserListConfig ListConfig;
 
+	// give external access to disable list update for perf
+	bool bAllowAutoUpdate = true;
 	bool bNeedsSort = false;
 	float AutoUpdatePeriod = 5.f;
 	FDelegateHandle UpdateTickerHandle;
diff --git a/Engine/Plugins/Online/OnlineFramework/Source/Party/Public/User/ISocialUserList.h b/Engine/Plugins/Online/OnlineFramework/Source/Party/Public/User/ISocialUserList.h
--- a/Engine/Plugins/Online/OnlineFramework/Source/Party/Public/User/ISocialUserList.h
+++ b/Engine/Plugins/Online/OnlineFramework/Source/Party/Public/User/ISocialUserList.h
@@ -35,6 +35,7 @@ public:
 	ESocialUserStateFlags RequiredPresenceFlags = ESocialUserStateFlags::None;
 	ESocialUserStateFlags ForbiddenPresenceFlags = ESocialUserStateFlags::None;
 	FOnCustomFilterUser OnCustomFilterUser;
+	bool bRequireAutoUpdate = false;
 };
 
 class ISocialUserList
@@ -57,6 +58,9 @@ public:
 	/** Trigger an update of the list immediately, regardless of auto update period */
 	virtual void UpdateNow() = 0;
 
+	/** Give external overwrite to disable list auto update for perf */
+	virtual void SetAllowAutoUpdate(bool bIsEnabled) = 0;
+
 	/** Sets the period at which to update the list with all users that  */
 	virtual void SetAutoUpdatePeriod(float InAutoUpdatePeriod) = 0;
 };
\ No newline at end of file
