Fixing logical bug that occured after slash command refactor. Condition to only execute the command when the required user was valid had a ! in the condition when it shouldn't.

#tests works as intended after removing the accidental ! sign

diff --git a/Engine/Plugins/Online/OnlineFramework/Source/Party/Private/Chat/ChatSlashCommands.cpp b/Engine/Plugins/Online/OnlineFramework/Source/Party/Private/Chat/ChatSlashCommands.cpp
--- a/Engine/Plugins/Online/OnlineFramework/Source/Party/Private/Chat/ChatSlashCommands.cpp
+++ b/Engine/Plugins/Online/OnlineFramework/Source/Party/Private/Chat/ChatSlashCommands.cpp
@@ -119,7 +119,7 @@ bool FRegisteredSlashCommands::NotifyUserTextChanged(const FText& InText)
 					{
 						if (Cmd->RequiresUserForExecution())
 						{
-							if (!AutoCompleteDatum->OptionalTargetUser.IsValid())
+							if (AutoCompleteDatum->OptionalTargetUser.IsValid())
 							{
 								Cmd->ExecuteSlashCommand(AutoCompleteDatum->OptionalTargetUser.Get());
 								return true;
