[NO NON-PII COMMIT MESSAGE]

diff --git a/Engine/Plugins/Online/OnlineFramework/Source/Party/Private/Chat/SocialGroupChannel.cpp b/Engine/Plugins/Online/OnlineFramework/Source/Party/Private/Chat/SocialGroupChannel.cpp
--- a/Engine/Plugins/Online/OnlineFramework/Source/Party/Private/Chat/SocialGroupChannel.cpp
+++ b/Engine/Plugins/Online/OnlineFramework/Source/Party/Private/Chat/SocialGroupChannel.cpp
@@ -1,4 +1,4 @@
-// Copyright 1998-2016 Epic Games, Inc. All Rights Reserved.
+// Copyright 1998-2019 Epic Games, Inc. All Rights Reserved.
 
 #include "Chat/SocialGroupChannel.h"
 #include "User/SocialUser.h"
diff --git a/Engine/Plugins/Online/OnlineFramework/Source/Party/Private/SocialDebugTools.cpp b/Engine/Plugins/Online/OnlineFramework/Source/Party/Private/SocialDebugTools.cpp
--- a/Engine/Plugins/Online/OnlineFramework/Source/Party/Private/SocialDebugTools.cpp
+++ b/Engine/Plugins/Online/OnlineFramework/Source/Party/Private/SocialDebugTools.cpp
@@ -1,4 +1,4 @@
-// Copyright 1998-2018 Epic Games, Inc. All Rights Reserved.
+// Copyright 1998-2019 Epic Games, Inc. All Rights Reserved.
 
 #include "SocialDebugTools.h"
 
diff --git a/Engine/Plugins/Online/OnlineFramework/Source/Party/Public/Chat/SocialGroupChannel.h b/Engine/Plugins/Online/OnlineFramework/Source/Party/Public/Chat/SocialGroupChannel.h
--- a/Engine/Plugins/Online/OnlineFramework/Source/Party/Public/Chat/SocialGroupChannel.h
+++ b/Engine/Plugins/Online/OnlineFramework/Source/Party/Public/Chat/SocialGroupChannel.h
@@ -1,4 +1,4 @@
-// Copyright 1998-2018 Epic Games, Inc. All Rights Reserved.
+// Copyright 1998-2019 Epic Games, Inc. All Rights Reserved.
 
 #pragma once
 
diff --git a/Engine/Plugins/Online/OnlineFramework/Source/Party/Public/SocialDebugTools.h b/Engine/Plugins/Online/OnlineFramework/Source/Party/Public/SocialDebugTools.h
--- a/Engine/Plugins/Online/OnlineFramework/Source/Party/Public/SocialDebugTools.h
+++ b/Engine/Plugins/Online/OnlineFramework/Source/Party/Public/SocialDebugTools.h
@@ -1,4 +1,4 @@
-// Copyright 1998-2018 Epic Games, Inc. All Rights Reserved.
+// Copyright 1998-2019 Epic Games, Inc. All Rights Reserved.
 
 #pragma once
 
diff --git a/Engine/Plugins/Online/OnlineSubsystem/Source/Public/Interfaces/OnlineDeviceNotificationInterface.h b/Engine/Plugins/Online/OnlineSubsystem/Source/Public/Interfaces/OnlineDeviceNotificationInterface.h
--- a/Engine/Plugins/Online/OnlineSubsystem/Source/Public/Interfaces/OnlineDeviceNotificationInterface.h
+++ b/Engine/Plugins/Online/OnlineSubsystem/Source/Public/Interfaces/OnlineDeviceNotificationInterface.h
@@ -1,4 +1,4 @@
-// Copyright 1998-2018 Epic Games, Inc. All Rights Reserved.
+// Copyright 1998-2019 Epic Games, Inc. All Rights Reserved.
 
 #pragma once
 
