Also added SocialSettings.MaxPartySize cvar for those all-the-times you might forget to set the config override value

diff --git a/Engine/Plugins/Online/OnlineFramework/Source/Party/Private/SocialSettings.cpp b/Engine/Plugins/Online/OnlineFramework/Source/Party/Private/SocialSettings.cpp
--- a/Engine/Plugins/Online/OnlineFramework/Source/Party/Private/SocialSettings.cpp
+++ b/Engine/Plugins/Online/OnlineFramework/Source/Party/Private/SocialSettings.cpp
@@ -3,6 +3,14 @@
 #include "SocialSettings.h"
 #include "SocialManager.h"
 
+#if !UE_BUILD_SHIPPING
+int32 MaxPartySizeOverride = INDEX_NONE;
+FAutoConsoleVariableRef CVarMaxPartySize(
+	TEXT("SocialSettings.MaxPartySize"),
+	MaxPartySizeOverride,
+	TEXT("Override the maximum persistent party size allowed by the social system"));
+#endif
+
 USocialSettings::USocialSettings()
 {
 	// Switch is the only default supported OSS that does not itself support multiple environments
@@ -30,6 +38,13 @@ bool USocialSettings::ShouldPreferPlatformInvites()
 
 int32 USocialSettings::GetDefaultMaxPartySize()
 {
+#if !UE_BUILD_SHIPPING
+	if (MaxPartySizeOverride > 0)
+	{
+		return MaxPartySizeOverride;
+	}
+#endif
+
 	const USocialSettings& SettingsCDO = *GetDefault<USocialSettings>();
 	return SettingsCDO.DefaultMaxPartySize;
 }
