Reworked the way we process Exec commands in SocialDebugTools to shunt everything after "Social Debug Context=" to a RunCommand() function so subclasses can generically act on an array of instance names.

Further optimization here would be to resolve all contexts before processing commands and pass an array of FInstanceContext&

#tests Compile DebugGame Editor Win64 / PIE

diff --git a/Engine/Plugins/Online/OnlineFramework/Source/Party/Private/SocialDebugTools.cpp b/Engine/Plugins/Online/OnlineFramework/Source/Party/Private/SocialDebugTools.cpp
--- a/Engine/Plugins/Online/OnlineFramework/Source/Party/Private/SocialDebugTools.cpp
+++ b/Engine/Plugins/Online/OnlineFramework/Source/Party/Private/SocialDebugTools.cpp
@@ -30,12 +30,12 @@ void USocialDebugTools::Shutdown()
 	Contexts.Empty();
 }
 
-void USocialDebugTools::PrintExecUsage()
+void USocialDebugTools::PrintExecUsage() const
 {
 	UE_LOG(LogParty, Log, TEXT("Usage SOCIAL DEBUG CONTEXT=<Name> <COMMAND> <PARAMS>"));
 }
 
-void USocialDebugTools::PrintExecCommands()
+void USocialDebugTools::PrintExecCommands() const
 {
 	UE_LOG(LogParty, Log, TEXT("LOGIN <Id> <Auth>"));
 	UE_LOG(LogParty, Log, TEXT("LOGOUT"));
@@ -83,108 +83,9 @@ bool USocialDebugTools::Exec(class UWorld* InWorld, const TCHAR* Cmd, FOutputDev
 				TargetInstances.Add(Instance);
 			}
 
-			if (FParse::Command(&Cmd, TEXT("LOGIN")))
-			{
-				if (bAllInstancesRequested)
-				{
-					UE_LOG(LogParty, Log, TEXT("CONTEXT=ALL cannot be used for the LOGIN command!"));
-				}
-				else
-				{
-					const FString Id = FParse::Token(Cmd, false);
-					const FString Auth = FParse::Token(Cmd, false);
-
-					Login(Instance, FOnlineAccountCredentials(TEXT("epic"), Id, Auth), FLoginComplete::CreateLambda([this, Instance](bool bSuccess)
-					{
-						UE_LOG(LogParty, Display, TEXT("Login OSS context[%s] %s"), *Instance, *LexToString(bSuccess));
-					}));
-				}
-			}
-			else if (FParse::Command(&Cmd, TEXT("LOGOUT")))
-			{
-				for (const FString& TargetInstance : TargetInstances)
-				{
-					Logout(TargetInstance, FLogoutComplete::CreateLambda([this, TargetInstance](bool bSuccess)
-					{
-						UE_LOG(LogParty, Display, TEXT("Logout OSS context[%s] %s"), *TargetInstance, *LexToString(bSuccess));
-					}));
-				}
-			}
-			else if (FParse::Command(&Cmd, TEXT("JOINPARTY")))
-			{
-				const FString Id = FParse::Token(Cmd, false);
-				const FString Auth = FParse::Token(Cmd, false);
-				const FString FriendName = FParse::Token(Cmd, false);
-				
-				for( const FString& TargetInstance : TargetInstances )
-				{
-					Login(TargetInstance, FOnlineAccountCredentials(TEXT("epic"), Id, Auth), FLoginComplete::CreateLambda([this, TargetInstance, FriendName](bool bSuccess)
-					{
-						UE_LOG(LogParty, Display, TEXT("Login OSS context[%s] %s"), *TargetInstance, *LexToString(bSuccess));
-
-						if (bSuccess)
-						{
-							LeaveParty(TargetInstance, FLeavePartyComplete::CreateLambda([this, TargetInstance, FriendName](bool bLeavePartySuccess)
-							{
-								UE_LOG(LogParty, Display, TEXT("Leave party OSS context[%s] %s"), *TargetInstance, *LexToString(bLeavePartySuccess));
-
-								JoinParty(TargetInstance, FriendName, FJoinPartyComplete::CreateLambda([this, TargetInstance](bool bJoinPartySuccess)
-								{
-									UE_LOG(LogParty, Display, TEXT("Join party OSS context[%s] %s"), *TargetInstance, *LexToString(bJoinPartySuccess));
-								}));
-							}));
-						}
-					}));
-				}
-			}
-			else if (FParse::Command(&Cmd, TEXT("LEAVEPARTY")))
-			{
-				for (const FString& TargetInstance : TargetInstances)
-				{
-					LeaveParty(TargetInstance, FLeavePartyComplete::CreateLambda([this, TargetInstance](bool bSuccess)
-					{
-						UE_LOG(LogParty, Display, TEXT("Leave party OSS context[%s] %s"), *TargetInstance, *LexToString(bSuccess));
-					}));
-				}
-			}
-			else if (FParse::Command(&Cmd, TEXT("AUTOACCEPTFRIENDINVITES")))
-			{
-				const FString Enabled = FParse::Token(Cmd, false);
-				if (Enabled == TEXT("true") || Enabled == TEXT("1"))
-				{
-					bAutoAcceptFriendInvites = true;
-				}
-				else if (Enabled == TEXT("false") || Enabled == TEXT("0"))
-				{
-					bAutoAcceptFriendInvites = false;
-				}
-			}
-			else if (FParse::Command(&Cmd, TEXT("AUTOACCEPTPARTYINVITES")))
-			{
-				const FString Enabled = FParse::Token(Cmd, false);
-				if (Enabled == TEXT("true") || Enabled == TEXT("1"))
-				{
-					bAutoAcceptPartyInvites = true;
-				}
-				else if (Enabled == TEXT("false") || Enabled == TEXT("0"))
-				{
-					bAutoAcceptPartyInvites = false;
-				}
-			}
-			else if (FParse::Command(&Cmd, TEXT("DUMPPARTY")))
-			{
-				for (const FString& TargetInstance : TargetInstances)
-				{
-					UE_LOG(LogParty, Display, TEXT("---DUMPPARTY - party OSS context[%s]"), *TargetInstance);
-					IOnlineSubsystem* OnlineSub = GetContext(TargetInstance).GetOSS();
-					if (OnlineSub &&
-						OnlineSub->GetPartyInterface().IsValid())
-					{
-						OnlineSub->GetPartyInterface()->DumpPartyState();
-					}
-				}
-			}			
+			return RunCommand(Cmd, TargetInstances);
 		}
+
 		return true;
 	}
 	
@@ -564,6 +465,7 @@ USocialDebugTools::FInstanceContext& USocialDebugTools::GetContext(const FString
 	{
 		Context = &Contexts.Add(Instance, FInstanceContext(Instance, *this));
 		Context->Init();
+		NotifyContextInitialized(*Context);
 	}
 	return *Context;
 }
@@ -592,6 +494,121 @@ USocialDebugTools::FInstanceContext* USocialDebugTools::GetContextForUser(const
 	return Result;
 }
 
+bool USocialDebugTools::RunCommand(const TCHAR* Cmd, const TArray<FString>& TargetInstances)
+{
+	if (FParse::Command(&Cmd, TEXT("LOGIN")))
+	{
+		if (TargetInstances.Num() > 1)
+		{
+			UE_LOG(LogParty, Log, TEXT("CONTEXT=ALL cannot be used for the LOGIN command!"));
+		}
+		else
+		{
+			const FString Id = FParse::Token(Cmd, false);
+			const FString Auth = FParse::Token(Cmd, false);
+			FString InstanceName = TargetInstances[0];
+
+			Login(InstanceName, FOnlineAccountCredentials(TEXT("epic"), Id, Auth), FLoginComplete::CreateLambda([this, InstanceName](bool bSuccess)
+			{
+				UE_LOG(LogParty, Display, TEXT("Login OSS context[%s] %s"), *InstanceName, *LexToString(bSuccess));
+			}));
+		}
+		return true;
+	}
+	else if (FParse::Command(&Cmd, TEXT("LOGOUT")))
+	{
+		for (const FString& TargetInstance : TargetInstances)
+		{
+			Logout(TargetInstance, FLogoutComplete::CreateLambda([this, TargetInstance](bool bSuccess)
+			{
+				UE_LOG(LogParty, Display, TEXT("Logout OSS context[%s] %s"), *TargetInstance, *LexToString(bSuccess));
+			}));
+		}
+		return true;
+	}
+	else if (FParse::Command(&Cmd, TEXT("JOINPARTY")))
+	{
+		const FString Id = FParse::Token(Cmd, false);
+		const FString Auth = FParse::Token(Cmd, false);
+		const FString FriendName = FParse::Token(Cmd, false);
+
+		for (const FString& TargetInstance : TargetInstances)
+		{
+			Login(TargetInstance, FOnlineAccountCredentials(TEXT("epic"), Id, Auth), FLoginComplete::CreateLambda([this, TargetInstance, FriendName](bool bSuccess)
+			{
+				UE_LOG(LogParty, Display, TEXT("Login OSS context[%s] %s"), *TargetInstance, *LexToString(bSuccess));
+
+				if (bSuccess)
+				{
+					LeaveParty(TargetInstance, FLeavePartyComplete::CreateLambda([this, TargetInstance, FriendName](bool bLeavePartySuccess)
+					{
+						UE_LOG(LogParty, Display, TEXT("Leave party OSS context[%s] %s"), *TargetInstance, *LexToString(bLeavePartySuccess));
+
+						JoinParty(TargetInstance, FriendName, FJoinPartyComplete::CreateLambda([this, TargetInstance](bool bJoinPartySuccess)
+						{
+							UE_LOG(LogParty, Display, TEXT("Join party OSS context[%s] %s"), *TargetInstance, *LexToString(bJoinPartySuccess));
+						}));
+					}));
+				}
+			}));
+		}
+		return true;
+	}
+	else if (FParse::Command(&Cmd, TEXT("LEAVEPARTY")))
+	{
+		for (const FString& TargetInstance : TargetInstances)
+		{
+			LeaveParty(TargetInstance, FLeavePartyComplete::CreateLambda([this, TargetInstance](bool bSuccess)
+			{
+				UE_LOG(LogParty, Display, TEXT("Leave party OSS context[%s] %s"), *TargetInstance, *LexToString(bSuccess));
+			}));
+		}
+		return true;
+	}
+	else if (FParse::Command(&Cmd, TEXT("AUTOACCEPTFRIENDINVITES")))
+	{
+		const FString Enabled = FParse::Token(Cmd, false);
+		if (Enabled == TEXT("true") || Enabled == TEXT("1"))
+		{
+			bAutoAcceptFriendInvites = true;
+		}
+		else if (Enabled == TEXT("false") || Enabled == TEXT("0"))
+		{
+			bAutoAcceptFriendInvites = false;
+		}
+		return true;
+	}
+	else if (FParse::Command(&Cmd, TEXT("AUTOACCEPTPARTYINVITES")))
+	{
+		const FString Enabled = FParse::Token(Cmd, false);
+		if (Enabled == TEXT("true") || Enabled == TEXT("1"))
+		{
+			bAutoAcceptPartyInvites = true;
+		}
+		else if (Enabled == TEXT("false") || Enabled == TEXT("0"))
+		{
+			bAutoAcceptPartyInvites = false;
+		}
+		return true;
+	}
+	else if (FParse::Command(&Cmd, TEXT("DUMPPARTY")))
+	{
+		for (const FString& TargetInstance : TargetInstances)
+		{
+			UE_LOG(LogParty, Display, TEXT("---DUMPPARTY - party OSS context[%s]"), *TargetInstance);
+			IOnlineSubsystem* OnlineSub = GetContext(TargetInstance).GetOSS();
+			if (OnlineSub &&
+				OnlineSub->GetPartyInterface().IsValid())
+			{
+				OnlineSub->GetPartyInterface()->DumpPartyState();
+			}
+		}
+		return true;
+	}
+
+	return false;
+}
+
 void USocialDebugTools::FInstanceContext::Init()
 {
 	const FString OSSName = FName(MCP_SUBSYSTEM).ToString() + TEXT(":") + Name;
diff --git a/Engine/Plugins/Online/OnlineFramework/Source/Party/Public/SocialDebugTools.h b/Engine/Plugins/Online/OnlineFramework/Source/Party/Public/SocialDebugTools.h
--- a/Engine/Plugins/Online/OnlineFramework/Source/Party/Public/SocialDebugTools.h
+++ b/Engine/Plugins/Online/OnlineFramework/Source/Party/Public/SocialDebugTools.h
@@ -73,6 +73,11 @@ public:
 
 	FInstanceContext& GetContext(const FString& Instance);
 	FInstanceContext* GetContextForUser(const FUniqueNetId& UserId);
+
+protected:
+	virtual bool RunCommand(const TCHAR* Cmd, const TArray<FString>& TargetInstances);
+	virtual void NotifyContextInitialized(const FInstanceContext& Context) { }
+
 private:
 
 	bool bAutoAcceptFriendInvites;
@@ -82,8 +87,8 @@ private:
 
 	TSharedPtr<IOnlinePartyJoinInfo> GetDefaultPartyJoinInfo() const;
 	IOnlineSubsystem* GetDefaultOSS() const;
-	void PrintExecUsage();
-	void PrintExecCommands();
+	void PrintExecUsage() const;
+	virtual void PrintExecCommands() const;
 
 	// OSS callback handlers
 	void HandleFriendInviteReceived(const FUniqueNetId& LocalUserId, const FUniqueNetId& FriendId);
