Fix log line indicating there's no OSS party when there is no logged in user

diff --git a/Engine/Plugins/Online/OnlineFramework/Source/Party/Private/SocialManager.cpp b/Engine/Plugins/Online/OnlineFramework/Source/Party/Private/SocialManager.cpp
--- a/Engine/Plugins/Online/OnlineFramework/Source/Party/Private/SocialManager.cpp
+++ b/Engine/Plugins/Online/OnlineFramework/Source/Party/Private/SocialManager.cpp
@@ -352,7 +352,7 @@ void USocialManager::CreateParty(const FOnlinePartyTypeId& PartyTypeId, const FP
 		}
 		else
 		{
-			UE_LOG(LogParty, Warning, TEXT("Cannot create party of type [%d] - no PartyInterface available on the primary OSS [%s]"), PartyTypeId.GetValue(), *GetSocialOssName(ESocialSubsystem::Primary).ToString());
+			UE_LOG(LogParty, Warning, TEXT("Cannot create party of type [%d] - PartyInterface.IsValid=%s PrimaryLocalUserId.IsValid=%s primary OSS [%s]"), PartyTypeId.GetValue(), *LexToString(PartyInterface.IsValid()), *LexToString(PrimaryLocalUserId.IsValid()), *GetSocialOssName(ESocialSubsystem::Primary).ToString());
 			OnCreatePartyComplete.ExecuteIfBound(ECreatePartyCompletionResult::UnknownClientFailure);
 		}
 	}
