Added non-shipping "-maxpartysize=" command-line parameter to complement the cvar override previously added.

#tests Compile DebugGame Editor Win64

diff --git a/Engine/Plugins/Online/OnlineFramework/Source/Party/Private/SocialSettings.cpp b/Engine/Plugins/Online/OnlineFramework/Source/Party/Private/SocialSettings.cpp
--- a/Engine/Plugins/Online/OnlineFramework/Source/Party/Private/SocialSettings.cpp
+++ b/Engine/Plugins/Online/OnlineFramework/Source/Party/Private/SocialSettings.cpp
@@ -2,6 +2,7 @@
 
 #include "SocialSettings.h"
 #include "SocialManager.h"
+#include "Misc/CommandLine.h"
 
 #if !UE_BUILD_SHIPPING
 int32 MaxPartySizeOverride = INDEX_NONE;
@@ -43,9 +44,15 @@ int32 USocialSettings::GetDefaultMaxPartySize()
 	{
 		return MaxPartySizeOverride;
 	}
-#endif
 
 	const USocialSettings& SettingsCDO = *GetDefault<USocialSettings>();
+	static FString CommandLineOverridePartySize;
+	if (FParse::Value(FCommandLine::Get(), TEXT("MaxPartySize="), CommandLineOverridePartySize))
+	{
+		return FCString::Atoi(*CommandLineOverridePartySize);
+	}
+#endif
+
 	return SettingsCDO.DefaultMaxPartySize;
 }
 
