Fixed compile error from previous submit.

diff --git a/Engine/Plugins/Online/OnlineFramework/Source/Party/Private/SocialSettings.cpp b/Engine/Plugins/Online/OnlineFramework/Source/Party/Private/SocialSettings.cpp
--- a/Engine/Plugins/Online/OnlineFramework/Source/Party/Private/SocialSettings.cpp
+++ b/Engine/Plugins/Online/OnlineFramework/Source/Party/Private/SocialSettings.cpp
@@ -53,6 +53,7 @@ int32 USocialSettings::GetDefaultMaxPartySize()
 	}
 #endif
 
+	const USocialSettings& SettingsCDO = *GetDefault<USocialSettings>();
 	return SettingsCDO.DefaultMaxPartySize;
 }
 
