Fixing previous previous compile error.

diff --git a/Engine/Plugins/Online/OnlineFramework/Source/Party/Private/SocialSettings.cpp b/Engine/Plugins/Online/OnlineFramework/Source/Party/Private/SocialSettings.cpp
--- a/Engine/Plugins/Online/OnlineFramework/Source/Party/Private/SocialSettings.cpp
+++ b/Engine/Plugins/Online/OnlineFramework/Source/Party/Private/SocialSettings.cpp
@@ -45,7 +45,6 @@ int32 USocialSettings::GetDefaultMaxPartySize()
 		return MaxPartySizeOverride;
 	}
 
-	const USocialSettings& SettingsCDO = *GetDefault<USocialSettings>();
 	static FString CommandLineOverridePartySize;
 	if (FParse::Value(FCommandLine::Get(), TEXT("MaxPartySize="), CommandLineOverridePartySize))
 	{
