Added a protected accessor for the PartyConfiguration and promoted UpdatePartyConfig so that subclasses can issue updates when necessary.

#tests Compile DebugGame Editor Win64 / PIE

diff --git a/Engine/Plugins/Online/OnlineFramework/Source/Party/Public/Party/SocialParty.h b/Engine/Plugins/Online/OnlineFramework/Source/Party/Public/Party/SocialParty.h
--- a/Engine/Plugins/Online/OnlineFramework/Source/Party/Public/Party/SocialParty.h
+++ b/Engine/Plugins/Online/OnlineFramework/Source/Party/Public/Party/SocialParty.h
@@ -212,7 +212,9 @@ PARTY_SCOPE:
 
 protected:
 	virtual void InitializePartyInternal();
-	
+
+	FPartyConfiguration& GetCurrentConfiguration() { return CurrentConfig; }
+
 	/** Only called when a new party is being created by the local player and they are responsible for the rep data. Otherwise we just wait to receive it from the leader. */
 	virtual void InitializePartyRepData();
 	virtual FPartyPrivacySettings GetDesiredPrivacySettings() const;
@@ -269,6 +271,9 @@ protected:
 	UPROPERTY()
 		TSubclassOf<ASpectatorBeaconClient> SpectatorBeaconClientClass;
 
+	/** Apply local party configuration to the OSS party, optionally resetting the access key to the party in the process */
+	void UpdatePartyConfig(bool bResetAccessKey = false);
+
 private:
 	UPartyMember* GetOrCreatePartyMember(const FUniqueNetId& MemberId);
 	void PumpApprovalQueue();
@@ -281,8 +286,6 @@ private:
 
 	void HandlePreClientTravel(const FString& PendingURL, ETravelType TravelType, bool bIsSeamlessTravel);
 
-	/** Apply local party configuration to the OSS party, optionally resetting the access key to the party in the process */
-	void UpdatePartyConfig(bool bResetAccessKey = false);
 
 	UPartyMember* GetMemberInternal(const FUniqueNetIdRepl& MemberId) const;
 
