Added a delegate that fires when the ALobbyBeaconPlayerState::UniqueId is replicated.

diff --git a/Engine/Plugins/Online/OnlineFramework/Source/Lobby/Private/LobbyBeaconPlayerState.cpp b/Engine/Plugins/Online/OnlineFramework/Source/Lobby/Private/LobbyBeaconPlayerState.cpp
--- a/Engine/Plugins/Online/OnlineFramework/Source/Lobby/Private/LobbyBeaconPlayerState.cpp
+++ b/Engine/Plugins/Online/OnlineFramework/Source/Lobby/Private/LobbyBeaconPlayerState.cpp
@@ -45,6 +45,14 @@ bool ALobbyBeaconPlayerState::IsValid() const
 	return UniqueId.IsValid();
 }
 
+void ALobbyBeaconPlayerState::OnRep_UniqueId()
+{
+	if (UniqueIdReplicatedEvent.IsBound())
+	{
+		UniqueIdReplicatedEvent.Broadcast(UniqueId);
+	}
+}
+
 void ALobbyBeaconPlayerState::OnRep_PartyOwner()
 {
 	if (PartyOwnerChangedEvent.IsBound())
diff --git a/Engine/Plugins/Online/OnlineFramework/Source/Lobby/Public/LobbyBeaconPlayerState.h b/Engine/Plugins/Online/OnlineFramework/Source/Lobby/Public/LobbyBeaconPlayerState.h
--- a/Engine/Plugins/Online/OnlineFramework/Source/Lobby/Public/LobbyBeaconPlayerState.h
+++ b/Engine/Plugins/Online/OnlineFramework/Source/Lobby/Public/LobbyBeaconPlayerState.h
@@ -33,7 +33,7 @@ class LOBBY_API ALobbyBeaconPlayerState : public AInfo
 	FText DisplayName;
 
 	/** Player unique id */
-	UPROPERTY(Replicated)
+	UPROPERTY(ReplicatedUsing = OnRep_UniqueId)
 	FUniqueNetIdRepl UniqueId;
 
     /** Party owner id */
@@ -55,6 +55,8 @@ class LOBBY_API ALobbyBeaconPlayerState : public AInfo
 	 */
 	DECLARE_EVENT_OneParam( ALobbyBeaconPlayerState, FOnPlayerStateChanged, const FUniqueNetIdRepl& /** UniqueId */ );
 
+	/** @return delegate fired when the unique id of the player has been replicated */
+	FORCEINLINE FOnPlayerStateChanged& OnUniqueIdReplicated() { return UniqueIdReplicatedEvent; }
 	/** @return delegate fired when the state of the player has changed in some way */
 	FORCEINLINE FOnPlayerStateChanged& OnPlayerStateChanged() { return PlayerStateChangedEvent; }
 	/** @return delegate fired when the party owner of this player has changed */
@@ -62,6 +64,10 @@ class LOBBY_API ALobbyBeaconPlayerState : public AInfo
 
 protected:
 	
+	/** Unique Id has replicated */
+	UFUNCTION()
+	void OnRep_UniqueId();
+
 	/** Party owner has changed */
 	UFUNCTION()
 	void OnRep_PartyOwner();
@@ -71,7 +77,8 @@ protected:
 	void OnRep_InLobby();
 
 private:
-
+	/** Delegate fired when player unique id is replicated */
+	FOnPlayerStateChanged UniqueIdReplicatedEvent;
 	/** Delegate fired when player state changes */
 	FOnPlayerStateChanged PlayerStateChangedEvent;
 	/** Delegate fired when party owner changes */
