Added support for games to provide game-specific SocialQueries during their user initializations.

Promoted some SocialUser initialization functions and marked virtual so they can be subclassed by games.

Promoted FSocialQuery to public-facing in its plugin so it can be subclassed as well.

Added IAvatarProvider interface derived from IModularFeature so we can access it without subsystem-specific knowledge after it's registered.

#tests Compile DebugGame Editor Win64

diff --git a/Engine/Plugins/Online/OnlineFramework/Source/Party/Private/SocialQuery.h b/Engine/Plugins/Online/OnlineFramework/Source/Party/Private/SocialQuery.h
deleted file mode 100644
diff --git a/Engine/Plugins/Online/OnlineFramework/Source/Party/Public/SocialQuery.h b/Engine/Plugins/Online/OnlineFramework/Source/Party/Public/SocialQuery.h
new file mode 100644
--- /dev/null
+++ b/Engine/Plugins/Online/OnlineFramework/Source/Party/Public/SocialQuery.h
@@ -0,0 +1,157 @@
+// Copyright 1998-2019 Epic Games, Inc. All Rights Reserved.
+#pragma once
+
+#include "CoreMinimal.h"
+#include "Containers/Ticker.h"
+#include "OnlineSubsystem.h"
+#include "SocialToolkit.h"
+
+DECLARE_DELEGATE_TwoParams(FOnQueryCompleted, FName, const TSharedRef<class FSocialQueryBase>&);
+
+class FSocialQueryBase : public TSharedFromThis<FSocialQueryBase>
+{
+public:
+	virtual ~FSocialQueryBase() {}
+	virtual void ExecuteQuery() = 0;
+
+	bool HasExecuted() const { return bHasExecuted; }
+	ESocialSubsystem GetSubsystemType() const { return SubsystemType; }
+	const USocialToolkit* GetOwningToolkit() const { return Toolkit.Get(); }
+
+protected:
+	TWeakObjectPtr<const USocialToolkit> Toolkit;
+	ESocialSubsystem SubsystemType;
+	bool bHasExecuted = false;
+	FOnQueryCompleted OnQueryCompleted;
+};
+
+template <typename QueryUserIdT, typename... CompletionCallbackArgs>
+class TSocialQuery : public FSocialQueryBase
+{
+public:
+	using FQueryId = QueryUserIdT;
+	using FOnQueryComplete = TBaseDelegate<void, ESocialSubsystem, bool, CompletionCallbackArgs...>;
+
+	// All subclasses of TSocialQuery must implement this static method
+	// Intentionally not implemented here to catch errors at compile time
+	static FName GetQueryId();
+
+	virtual void AddUserId(const QueryUserIdT& UserId, const FOnQueryComplete& QueryCompleteHandler)
+	{
+		CompletionCallbacksByUserId.Add(UserId, QueryCompleteHandler);
+	}
+
+	template <typename OtherQueryT>
+	bool operator==(const TSharedRef<OtherQueryT>& OtherQuery) const
+	{
+		return TIsDerivedFrom<OtherQueryT, TSocialQuery>::IsDerived &&
+			OtherQuery->Subsystem == SubsystemType &&
+			OtherQuery->Toolkit == Toolkit;
+	}
+
+protected:
+	friend class FSocialQueryManager;
+	TSocialQuery() {}
+	void Initialize(const USocialToolkit& InToolkit, ESocialSubsystem InSubsystemType, const FOnQueryCompleted& InOnQueryCompleted)
+	{
+		Toolkit = &InToolkit;
+		SubsystemType = InSubsystemType;
+		OnQueryCompleted = InOnQueryCompleted;
+	}
+
+	inline IOnlineSubsystem* GetOSS() const
+	{
+		return Toolkit.IsValid() ? Toolkit->GetSocialOss(SubsystemType) : nullptr;
+	}
+
+	static TArray<TSharedRef<TSocialQuery>> CurrentQueries;
+	TMap<QueryUserIdT, FOnQueryComplete> CompletionCallbacksByUserId;
+};
+
+class FSocialQueryManager
+{
+public:
+	template <typename SocialQueryT>
+	static TSharedRef<SocialQueryT> GetQuery(const USocialToolkit& Toolkit, ESocialSubsystem SubsystemType)
+	{
+		return Get().GetQueryInternal<SocialQueryT>(Toolkit, SubsystemType);
+	}
+
+	template <typename SocialQueryT>
+	static void AddUserId(const USocialToolkit& Toolkit,
+		ESocialSubsystem InSubsystem, 
+		const typename SocialQueryT::FQueryId& InQueryId,
+		const typename SocialQueryT::FOnQueryComplete& OnQueryCompleteHandler)
+	{
+		Get().GetQueryInternal<SocialQueryT>(Toolkit, InSubsystem)->AddUserId(InQueryId, OnQueryCompleteHandler);
+	}
+
+	bool HandleExecuteQueries(float)
+	{
+		// Execute all pending queries
+		TArray<TArray<TSharedRef<FSocialQueryBase>>> AllQueries;
+		CurrentQueriesById.GenerateValueArray(AllQueries);
+		
+		for (const TArray<TSharedRef<FSocialQueryBase>>& Queries : AllQueries)
+		{
+			for (const TSharedRef<FSocialQueryBase>& Query : Queries)
+			{
+				if (!Query->HasExecuted())
+				{
+					Query->ExecuteQuery();
+				}
+			}
+		}
+
+		TickExecuteHandle.Reset();
+		return false;
+	}
+
+private:
+	FSocialQueryManager() {}
+
+	static FSocialQueryManager& Get()
+	{
+		static FSocialQueryManager SingletonInstance;
+		return SingletonInstance;
+	}
+
+	template <typename SocialQueryT>
+	TSharedRef<SocialQueryT> GetQueryInternal(const USocialToolkit& Toolkit, ESocialSubsystem SubsystemType)
+	{
+		const FName QueryId = SocialQueryT::GetQueryId();
+
+		TArray<TSharedRef<FSocialQueryBase>>& Queries = CurrentQueriesById.FindOrAdd(QueryId);
+		for (const TSharedRef<FSocialQueryBase>& Query : Queries)
+		{
+			if (Query->GetSubsystemType() == SubsystemType && Query->GetOwningToolkit() == &Toolkit && !Query->HasExecuted())
+			{
+				return StaticCastSharedRef<SocialQueryT>(Query);
+			}
+		}
+
+		// No matching query found, so make a new one
+		TSharedRef<SocialQueryT> NewQuery = MakeShareable(new SocialQueryT);
+		NewQuery->Initialize(Toolkit, SubsystemType, FOnQueryCompleted::CreateRaw(this, &FSocialQueryManager::HandleQueryComplete));
+		Queries.Add(NewQuery);
+
+		// If we aren't already registered to execute our queries next tick, do so now
+		if (!TickExecuteHandle.IsValid())
+		{
+			TickExecuteHandle = FTicker::GetCoreTicker().AddTicker(FTickerDelegate::CreateRaw(this, &FSocialQueryManager::HandleExecuteQueries));
+		}
+
+		return NewQuery;
+	}
+
+	void HandleQueryComplete(FName QueryId, const TSharedRef<FSocialQueryBase>& Query)
+	{
+		if (TArray<TSharedRef<FSocialQueryBase>>* Queries = CurrentQueriesById.Find(QueryId))
+		{
+			Queries->Remove(Query);
+		}
+	}
+
+	FDelegateHandle TickExecuteHandle;
+	TMap<FName, TArray<TSharedRef<FSocialQueryBase>>> CurrentQueriesById;
+};
\ No newline at end of file
diff --git a/Engine/Plugins/Online/OnlineFramework/Source/Party/Public/User/SocialUser.h b/Engine/Plugins/Online/OnlineFramework/Source/Party/Public/User/SocialUser.h
--- a/Engine/Plugins/Online/OnlineFramework/Source/Party/Public/User/SocialUser.h
+++ b/Engine/Plugins/Online/OnlineFramework/Source/Party/Public/User/SocialUser.h
@@ -127,15 +127,16 @@ PARTY_SCOPE:
 
 protected:
 	virtual void OnPresenceChangedInternal(ESocialSubsystem SubsystemType);
+	virtual void SetSubsystemId(ESocialSubsystem SubsystemType, const FUniqueNetIdRepl& SubsystemId);
+	int32 NumPendingQueries = 0;
 
+	void TryBroadcastInitializationComplete();
 private:
-	void SetSubsystemId(ESocialSubsystem SubsystemType, const FUniqueNetIdRepl& SubsystemId);
 	
 	void SetUserInfo(ESocialSubsystem SubsystemType, const TSharedRef<FOnlineUser>& UserInfo);
 	void HandleQueryUserInfoComplete(ESocialSubsystem SubsystemType, bool bWasSuccessful, const TSharedPtr<FOnlineUser>& UserInfo);
 	
 private:
-	void TryBroadcastInitializationComplete();
 
 	struct FSubsystemUserInfo
 	{
@@ -162,7 +163,6 @@ private:
 	};
 	FSubsystemUserInfo& FindOrCreateSubsystemInfo(const FUniqueNetIdRepl& SubsystemId, ESocialSubsystem SubsystemType);
 
-	int32 NumPendingQueries = 0;
 	bool bIsInitialized = false;
 
 	TSharedPtr<const IOnlinePartyJoinInfo> PersistentPartyInfo;
diff --git a/Engine/Plugins/Online/OnlineSubsystem/Source/Public/Interfaces/AvatarProvider.h b/Engine/Plugins/Online/OnlineSubsystem/Source/Public/Interfaces/AvatarProvider.h
new file mode 100644
--- /dev/null
+++ b/Engine/Plugins/Online/OnlineSubsystem/Source/Public/Interfaces/AvatarProvider.h
@@ -0,0 +1,49 @@
+// Copyright 1998-2019 Epic Games, Inc. All Rights Reserved.
+
+#pragma once
+
+#include "CoreMinimal.h"
+#include "Features/IModularFeature.h"
+
+class FUniqueNetId;
+
+
+struct FAvatarInfo : public TSharedFromThis<FAvatarInfo>
+{
+public:
+	FAvatarInfo(const TSharedPtr<const FUniqueNetId>& InUserId)
+		: UserId(InUserId.ToSharedRef())
+	{ }
+
+	FAvatarInfo(const TSharedPtr<const FUniqueNetId> InUserId, TMap<FString, FString>&& InAvatarInfoPairs)
+		: UserId(InUserId.ToSharedRef()), AvatarInfoPairs(MoveTemp(InAvatarInfoPairs))
+	{ }
+
+	TSharedRef<const FUniqueNetId> UserId;
+	TMap<FString, FString> AvatarInfoPairs;
+};
+
+/**
+ * Interface for a class that can provide support for querying information about an avatar associated with a user by UniqueNetId
+ */
+class IAvatarProvider : public IModularFeature
+{
+public:
+	DECLARE_DELEGATE_OneParam(FOnQueryAvatarInfoComplete, const FString& /*ErrorStr*/)
+
+	virtual ~IAvatarProvider() = default;
+
+	/**
+	 * Get the name of the modular feature, to be used to get the implementations
+	 * @return name of the modular feature
+	 */
+	static FName GetModularFeatureName()
+	{
+		static FName FeatureName = FName(TEXT("AvatarProvider"));
+		return FeatureName;
+	}
+
+	virtual void QueryAvatarInfo(const FUniqueNetId& LocalUserId, const TArray<TSharedRef<const FUniqueNetId>>& InUserIds, const FOnQueryAvatarInfoComplete& CompletionDelegate) = 0;
+
+	virtual TSharedPtr<const FAvatarInfo> GetAvatarInfo(const FUniqueNetId& InUserId) const = 0;
+};
