Created PartyMember accessor for the OSS-level party member so subgames can access member attribute map directly.

Added SocialManager accessor for SocialDebugTools.

diff --git a/Engine/Plugins/Online/OnlineFramework/Source/Party/Private/SocialDebugTools.cpp b/Engine/Plugins/Online/OnlineFramework/Source/Party/Private/SocialDebugTools.cpp
--- a/Engine/Plugins/Online/OnlineFramework/Source/Party/Private/SocialDebugTools.cpp
+++ b/Engine/Plugins/Online/OnlineFramework/Source/Party/Private/SocialDebugTools.cpp
@@ -11,6 +11,7 @@
 #include "Interfaces/OnlinePartyInterface.h"
 #include "Party/SocialParty.h"
 #include "Party/PartyMember.h"
+#include "SocialManager.h"
 
 USocialDebugTools::USocialDebugTools()
 	: bAutoAcceptFriendInvites(true)
@@ -426,6 +427,13 @@ IOnlineSubsystem* USocialDebugTools::GetDefaultOSS() const
 	
 }
 
+USocialManager& USocialDebugTools::GetSocialManager() const
+{
+	USocialManager* OuterSocialManager = GetTypedOuter<USocialManager>();
+	check(OuterSocialManager);
+	return *OuterSocialManager;
+}
+
 TSharedPtr<IOnlinePartyJoinInfo> USocialDebugTools::GetDefaultPartyJoinInfo() const
 {
 	TSharedPtr<IOnlinePartyJoinInfo> Result;
diff --git a/Engine/Plugins/Online/OnlineFramework/Source/Party/Public/Party/PartyMember.h b/Engine/Plugins/Online/OnlineFramework/Source/Party/Public/Party/PartyMember.h
--- a/Engine/Plugins/Online/OnlineFramework/Source/Party/Public/Party/PartyMember.h
+++ b/Engine/Plugins/Online/OnlineFramework/Source/Party/Public/Party/PartyMember.h
@@ -119,6 +119,8 @@ protected:
 	FPartyMemberDataReplicator MemberDataReplicator;
 	void SetMemberConnectionStatus(EMemberConnectionStatus InMemberConnectionStatus);
 
+	TSharedPtr<const FOnlinePartyMember> GetOSSPartyMember() const { return OssPartyMember; }
+
 private:
 	void HandleSocialUserInitialized(USocialUser& InitializedUser);
 	void HandleMemberConnectionStatusChanged(const FUniqueNetId& ChangedUserId, const EMemberConnectionStatus NewMemberConnectionStatus, const EMemberConnectionStatus PreviousMemberConnectionStatus);
diff --git a/Engine/Plugins/Online/OnlineFramework/Source/Party/Public/SocialDebugTools.h b/Engine/Plugins/Online/OnlineFramework/Source/Party/Public/SocialDebugTools.h
--- a/Engine/Plugins/Online/OnlineFramework/Source/Party/Public/SocialDebugTools.h
+++ b/Engine/Plugins/Online/OnlineFramework/Source/Party/Public/SocialDebugTools.h
@@ -19,6 +19,8 @@ class PARTY_API USocialDebugTools : public UObject, public FExec
 	static const int32 LocalUserNum = 0;
 
 public:
+	USocialManager& GetSocialManager() const;
+
 	// FExec
 	virtual bool Exec(class UWorld* InWorld, const TCHAR* Cmd, FOutputDevice& Out) override;
 
