Show "experiencing issues" message for 1 minute when party hotfix occurs that changes party versions.

Adds skeleton for future PartySystemState to be added later.

diff --git a/Engine/Plugins/Online/OnlineFramework/Source/Party/Private/Party/SocialParty.cpp b/Engine/Plugins/Online/OnlineFramework/Source/Party/Private/Party/SocialParty.cpp
--- a/Engine/Plugins/Online/OnlineFramework/Source/Party/Private/Party/SocialParty.cpp
+++ b/Engine/Plugins/Online/OnlineFramework/Source/Party/Private/Party/SocialParty.cpp
@@ -416,7 +416,7 @@ void USocialParty::InitializePartyInternal()
 	PartyInterface->AddOnPartyMemberDataReceivedDelegate_Handle(FOnPartyMemberDataReceivedDelegate::CreateUObject(this, &USocialParty::HandlePartyMemberDataReceived));
 	PartyInterface->AddOnPartyMemberPromotedDelegate_Handle(FOnPartyMemberPromotedDelegate::CreateUObject(this, &USocialParty::HandlePartyMemberPromoted));
 	PartyInterface->AddOnPartyMemberExitedDelegate_Handle(FOnPartyMemberExitedDelegate::CreateUObject(this, &USocialParty::HandlePartyMemberExited));
-	
+	PartyInterface->AddOnPartySystemStateChangeDelegate_Handle(FOnPartySystemStateChangeDelegate::CreateUObject(this, &USocialParty::HandlePartySystemStateChange));
 	// Create a UPartyMember for every existing member on the OSS party
 	TArray<TSharedRef<FOnlinePartyMember>> OssPartyMembers;
 	PartyInterface->GetPartyMembers(*OwningLocalUserId, GetPartyId(), OssPartyMembers);
@@ -920,6 +920,21 @@ void USocialParty::HandlePartyMemberExited(const FUniqueNetId& LocalUserId, cons
 	}
 }
 
+void USocialParty::HandlePartySystemStateChange(const FUniqueNetId& LocalUserId, EPartySystemState NewState)
+{
+	UE_LOG(LogParty, VeryVerbose, TEXT("Party [%s] received notification of a party system state change to [%d]"), *ToDebugString(), (int32)NewState);
+	if (NewState == EPartySystemState::RequestingShutdown)
+	{
+		// Need to display message
+		SetIsRequestingShutdown(true);
+
+		//set timer to turn this off in a minute?
+		UWorld* World = GetWorld();
+		FTimerHandle DummyHandle;
+		World->GetTimerManager().SetTimer(DummyHandle, FTimerDelegate::CreateLambda([this]() { SetIsRequestingShutdown(false); }), 60.0f, false);
+	}
+}
+
 FChatRoomId USocialParty::GetChatRoomId() const
 {
 	return ensure(OssParty.IsValid()) ? OssParty->RoomId : FChatRoomId();
@@ -987,7 +1002,7 @@ bool USocialParty::ShouldStayWithPartyOnExit() const
 
 bool USocialParty::IsPartyFunctionalityDegraded() const
 {
-	return bIsMissingXmppConnection.Get(false) || bIsMissingPlatformSession;
+	return bIsMissingXmppConnection.Get(false) || bIsMissingPlatformSession || bIsRequestingShutdown.Get(false);
 }
 
 int32 USocialParty::GetNumPartyMembers() const
@@ -1428,6 +1443,23 @@ void USocialParty::SetIsMissingXmppConnection(bool bInMissingXmppConnection)
 	}
 }
 
+void USocialParty::SetIsRequestingShutdown(bool bInRequestingShutdown)
+{
+	if (!bIsRequestingShutdown.IsSet() ||
+		bIsRequestingShutdown != bInRequestingShutdown)
+	{
+		UE_CLOG(bIsRequestingShutdown.IsSet(), LogParty, VeryVerbose, TEXT("Party [%s] is %s in a version transition"), *ToDebugString(), bInRequestingShutdown ? TEXT("now") : TEXT("no longer"));
+
+		const bool bWasPartyFunctionalityDegraded = IsPartyFunctionalityDegraded();
+		bIsRequestingShutdown = bInRequestingShutdown;
+		const bool bIsPartyFunctionalityDegraded = IsPartyFunctionalityDegraded();
+		if (bWasPartyFunctionalityDegraded != bIsPartyFunctionalityDegraded)
+		{
+			OnPartyFunctionalityDegradedChanged().Broadcast(bIsPartyFunctionalityDegraded);
+		}
+	}
+}
+
 void USocialParty::BeginLeavingParty(EMemberExitedReason Reason)
 {
 	if (!bIsLeavingParty)
diff --git a/Engine/Plugins/Online/OnlineFramework/Source/Party/Public/Party/SocialParty.h b/Engine/Plugins/Online/OnlineFramework/Source/Party/Public/Party/SocialParty.h
--- a/Engine/Plugins/Online/OnlineFramework/Source/Party/Public/Party/SocialParty.h
+++ b/Engine/Plugins/Online/OnlineFramework/Source/Party/Public/Party/SocialParty.h
@@ -282,6 +282,8 @@ private:
 	void BeginLeavingParty(EMemberExitedReason Reason);
 	void FinalizePartyLeave(EMemberExitedReason Reason);
 
+	void SetIsRequestingShutdown(bool bInRequestingShutdown);
+
 	void UpdatePlatformSessionLeader(FName PlatformOssName);
 
 	void HandlePreClientTravel(const FString& PendingURL, ETravelType TravelType, bool bIsSeamlessTravel);
@@ -313,6 +315,8 @@ private:	// Handlers
 
 	void HandleLeavePartyComplete(const FUniqueNetId& LocalUserId, const FOnlinePartyId& PartyId, ELeavePartyCompletionResult LeaveResult, FOnLeavePartyAttemptComplete OnAttemptComplete);
 	
+	void HandlePartySystemStateChange(const FUniqueNetId& LocalUserId, EPartySystemState NewState);
+
 private:
 	TSharedPtr<const FOnlineParty> OssParty;
 
@@ -373,6 +377,7 @@ private:
 
 	bool bIsLeavingParty = false;
 	bool bIsInitialized = false;
+	TOptional<bool> bIsRequestingShutdown;
 
 	mutable FLeavePartyEvent OnPartyLeaveBeginEvent;
 	mutable FLeavePartyEvent OnPartyLeftEvent;
diff --git a/Engine/Plugins/Online/OnlineSubsystem/Source/Public/Interfaces/OnlinePartyInterface.h b/Engine/Plugins/Online/OnlineSubsystem/Source/Public/Interfaces/OnlinePartyInterface.h
--- a/Engine/Plugins/Online/OnlineSubsystem/Source/Public/Interfaces/OnlinePartyInterface.h
+++ b/Engine/Plugins/Online/OnlineSubsystem/Source/Public/Interfaces/OnlinePartyInterface.h
@@ -627,6 +627,14 @@ enum class EKickMemberCompletionResult;
 enum class EPromoteMemberCompletionResult;
 enum class EInvitationResponse;
 
+enum class EPartySystemState
+{
+	Initializing = 0,
+	Initialized,
+	RequestingShutdown,
+	ShutDown,
+};
+
 ///////////////////////////////////////////////////////////////////
 // Completion delegates
 ///////////////////////////////////////////////////////////////////
@@ -927,6 +935,14 @@ PARTY_DECLARE_DELEGATETYPE(OnFillPartyJoinRequestData);
 DECLARE_MULTICAST_DELEGATE_ThreeParams(F_PREFIX(OnPartyAnalyticsEvent), const FUniqueNetId& /*LocalUserId*/, const FString& /*EventName*/, const TArray<FAnalyticsEventAttribute>& /*Attributes*/);
 PARTY_DECLARE_DELEGATETYPE(OnPartyAnalyticsEvent);
 
+/**
+* Notification of party system state change
+* @param LocalUserId - id associated with this event
+* @param NewState - new state this partysystem is in
+*/
+DECLARE_MULTICAST_DELEGATE_TwoParams(F_PREFIX(OnPartySystemStateChange), const FUniqueNetId& /*LocalUserId*/, EPartySystemState /*NewState*/);
+PARTY_DECLARE_DELEGATETYPE(OnPartySystemStateChange);
+
 /**
  * Interface definition for the online party services 
  * Allows for forming a party and communicating with party members
@@ -1577,6 +1593,13 @@ public:
 	 */
 	DEFINE_ONLINE_DELEGATE_THREE_PARAM(OnPartyAnalyticsEvent, const FUniqueNetId& /*LocalUserId*/, const FString& /*EventName*/, const TArray<FAnalyticsEventAttribute>& /*Attributes*/);
 
+	/**
+	* Notification of party system state change
+	* @param LocalUserId - id associated with this event
+	* @param NewState - new state this partysystem is in
+	*/
+	DEFINE_ONLINE_DELEGATE_TWO_PARAM(OnPartySystemStateChange, const FUniqueNetId& /*LocalUserId*/, EPartySystemState /*NewState*/);
+
 	/**
 	 * Dump out party state for all known parties
 	 */
