exposed OnHotfixAvailablityCheck notify with file list deltas
#rnx

diff --git a/Engine/Plugins/Online/OnlineFramework/Source/Hotfix/Private/OnlineHotfixManager.cpp b/Engine/Plugins/Online/OnlineFramework/Source/Hotfix/Private/OnlineHotfixManager.cpp
--- a/Engine/Plugins/Online/OnlineFramework/Source/Hotfix/Private/OnlineHotfixManager.cpp
+++ b/Engine/Plugins/Online/OnlineFramework/Source/Hotfix/Private/OnlineHotfixManager.cpp
@@ -423,6 +423,11 @@ void UOnlineHotfixManager::CheckAvailability(FOnHotfixAvailableComplete& InCompl
 	}
 }
 
+void UOnlineHotfixManager::OnHotfixAvailablityCheck(const TArray<FCloudFileHeader>& PendingChangedFiles, const TArray<FCloudFileHeader>& PendingRemoveFiles)
+{
+	// empty in base class
+}
+
 void UOnlineHotfixManager::OnEnumerateFilesForAvailabilityComplete(bool bWasSuccessful, const FString& ErrorStr, FOnHotfixAvailableComplete InCompletionDelegate)
 {
 	if (OnlineTitleFile.IsValid())
@@ -460,6 +465,8 @@ void UOnlineHotfixManager::OnEnumerateFilesForAvailabilityComplete(bool bWasSucc
 			Result = EHotfixResult::SuccessNoChange;
 		}
 
+		OnHotfixAvailablityCheck(ChangedHotfixFileList, RemovedHotfixFileList);
+
 		// Restore state to before the check
 		RemovedHotfixFileList.Empty();
 		ChangedHotfixFileList.Empty();
diff --git a/Engine/Plugins/Online/OnlineFramework/Source/Hotfix/Public/OnlineHotfixManager.h b/Engine/Plugins/Online/OnlineFramework/Source/Hotfix/Public/OnlineHotfixManager.h
--- a/Engine/Plugins/Online/OnlineFramework/Source/Hotfix/Public/OnlineHotfixManager.h
+++ b/Engine/Plugins/Online/OnlineFramework/Source/Hotfix/Public/OnlineHotfixManager.h
@@ -276,6 +276,9 @@ protected:
 		return FPaths::ProjectPersistentDownloadDir();
 	}
 
+	/** Notify used by CheckAvailability() */
+	virtual void OnHotfixAvailablityCheck(const TArray<FCloudFileHeader>& PendingChangedFiles, const TArray<FCloudFileHeader>& PendingRemoveFiles);
+
 	/** Finds the header associated with the file name */
 	FCloudFileHeader* GetFileHeaderFromDLName(const FString& FileName);
 
