Exposing option for children to override SocialParty checks for CanInvite, CanPromote, and CanKick

diff --git a/Engine/Plugins/Online/OnlineFramework/Source/Party/Private/Party/SocialParty.cpp b/Engine/Plugins/Online/OnlineFramework/Source/Party/Private/Party/SocialParty.cpp
--- a/Engine/Plugins/Online/OnlineFramework/Source/Party/Private/Party/SocialParty.cpp
+++ b/Engine/Plugins/Online/OnlineFramework/Source/Party/Private/Party/SocialParty.cpp
@@ -263,6 +263,11 @@ bool USocialParty::HasUserBeenInvited(const USocialUser& User) const
 }
 
 bool USocialParty::CanInviteUser(const USocialUser& User) const
+{
+	return CanInviteUserInternal(User);
+}
+
+bool USocialParty::CanInviteUserInternal(const USocialUser& User) const
 {
 	// Only users that are online can be invited
 	if (!User.IsOnline())
@@ -339,6 +344,11 @@ bool USocialParty::TryInviteUser(const USocialUser& UserToInvite)
 bool USocialParty::CanPromoteMember(const UPartyMember& PartyMember) const
 {
 	check(PartyMembersById.Contains(PartyMember.GetPrimaryNetId()));
+	return CanPromoteMemberInternal(PartyMember);
+}
+
+bool USocialParty::CanPromoteMemberInternal(const UPartyMember& PartyMember) const
+{
 	return IsLocalPlayerPartyLeader() && bIsMemberPromotionPossible && !PartyMember.IsPartyLeader();
 }
 
@@ -357,6 +367,11 @@ bool USocialParty::TryPromoteMember(const UPartyMember& PartyMember)
 bool USocialParty::CanKickMember(const UPartyMember& PartyMember) const
 {
 	check(PartyMembersById.Contains(PartyMember.GetPrimaryNetId()));
+	return CanKickMemberInternal(PartyMember);
+}
+
+bool USocialParty::CanKickMemberInternal(const UPartyMember& PartyMember) const
+{
 	return IsLocalPlayerPartyLeader() && !PartyMember.IsLocalPlayer();
 }
 
diff --git a/Engine/Plugins/Online/OnlineFramework/Source/Party/Public/Party/SocialParty.h b/Engine/Plugins/Online/OnlineFramework/Source/Party/Public/Party/SocialParty.h
--- a/Engine/Plugins/Online/OnlineFramework/Source/Party/Public/Party/SocialParty.h
+++ b/Engine/Plugins/Online/OnlineFramework/Source/Party/Public/Party/SocialParty.h
@@ -223,6 +223,11 @@ protected:
 	virtual void OnMemberCreatedInternal(UPartyMember& NewMember);
 	virtual void OnLeftPartyInternal(EMemberExitedReason Reason);
 
+	/** Virtual versions of the package-scoped "CanX" methods above, as a virtual declared within package scoping cannot link (exported public, imported protected) */
+	virtual bool CanInviteUserInternal(const USocialUser& User) const;
+	virtual bool CanPromoteMemberInternal(const UPartyMember& PartyMember) const;
+	virtual bool CanKickMemberInternal(const UPartyMember& PartyMember) const;
+
 	virtual void OnInviteSentInternal(ESocialSubsystem SubsystemType, const USocialUser& InvitedUser, bool bWasSuccessful);
 	
 	/** Determines the joinability of this party for a specific user requesting to join */
