Removing a couple erroneous (and useless) lines in UPartyMember::GetDisplayName()

diff --git a/Engine/Plugins/Online/OnlineFramework/Source/Party/Private/Party/PartyMember.cpp b/Engine/Plugins/Online/OnlineFramework/Source/Party/Private/Party/PartyMember.cpp
--- a/Engine/Plugins/Online/OnlineFramework/Source/Party/Private/Party/PartyMember.cpp
+++ b/Engine/Plugins/Online/OnlineFramework/Source/Party/Private/Party/PartyMember.cpp
@@ -147,10 +147,6 @@ USocialUser& UPartyMember::GetSocialUser() const
 
 FString UPartyMember::GetDisplayName() const
 {
-	if (SocialUser->IsInitialized())
-	{ 
-		SocialUser->GetDisplayName();
-	}
 	return OssPartyMember->GetDisplayName(GetRepData().GetPlatform());
 }
 
