bug fix for Matchmaking not being cancelled on party leader during party system hotfix

diff --git a/Engine/Plugins/Online/OnlineFramework/Source/Party/Public/Party/SocialParty.h b/Engine/Plugins/Online/OnlineFramework/Source/Party/Public/Party/SocialParty.h
--- a/Engine/Plugins/Online/OnlineFramework/Source/Party/Public/Party/SocialParty.h
+++ b/Engine/Plugins/Online/OnlineFramework/Source/Party/Public/Party/SocialParty.h
@@ -230,6 +230,8 @@ protected:
 
 	virtual void OnInviteSentInternal(ESocialSubsystem SubsystemType, const USocialUser& InvitedUser, bool bWasSuccessful);
 	
+	virtual void HandlePartySystemStateChange(EPartySystemState NewState);
+
 	/** Determines the joinability of this party for a specific user requesting to join */
 	virtual FPartyJoinApproval EvaluateJoinRequest(const FUniqueNetId& PlayerId, const FUserPlatform& Platform, const FOnlinePartyData& JoinData, bool bFromJoinRequest) const;
 
@@ -319,8 +321,6 @@ private:	// Handlers
 	void HandleReservationRequestComplete(EPartyReservationResult::Type ReservationResponse);
 
 	void HandleLeavePartyComplete(const FUniqueNetId& LocalUserId, const FOnlinePartyId& PartyId, ELeavePartyCompletionResult LeaveResult, FOnLeavePartyAttemptComplete OnAttemptComplete);
-	
-	void HandlePartySystemStateChange(EPartySystemState NewState);
 
 private:
 	TSharedPtr<const FOnlineParty> OssParty;
