Wait to broadcast Party Initialized until Replicated Data is received. If the local player is the party leader, record that the data is received when initializing it.

diff --git a/Engine/Plugins/Online/OnlineFramework/Source/Party/Private/Party/SocialParty.cpp b/Engine/Plugins/Online/OnlineFramework/Source/Party/Private/Party/SocialParty.cpp
--- a/Engine/Plugins/Online/OnlineFramework/Source/Party/Private/Party/SocialParty.cpp
+++ b/Engine/Plugins/Online/OnlineFramework/Source/Party/Private/Party/SocialParty.cpp
@@ -463,7 +463,7 @@ void USocialParty::TryFinishInitialization()
 		IOnlinePartyPtr PartyInterface = Online::GetPartyInterfaceChecked(GetWorld()); 
 		uint32 OSSMemberCount = PartyInterface->GetPartyMemberCount(*OwningLocalUserId, GetPartyId());
 
-		if (OSSMemberCount == PartyMembersById.Num())
+		if (OSSMemberCount == PartyMembersById.Num() && bHasReceivedRepData)
 		{
 			bIsInitialized = true;
 			GetSocialManager().NotifyPartyInitialized(*this);
@@ -497,6 +497,7 @@ void USocialParty::RefreshPublicJoinability()
 void USocialParty::InitializePartyRepData()
 {
 	UE_LOG(LogParty, Verbose, TEXT("Initializing rep data for party [%s]"), *ToDebugString());
+	bHasReceivedRepData = true;
 }
 
 FPartyPrivacySettings USocialParty::GetDesiredPrivacySettings() const
@@ -683,6 +684,11 @@ void USocialParty::HandlePartyDataReceived(const FUniqueNetId& LocalUserId, cons
 	{
 		check(PartyDataReplicator.IsValid());
 		PartyDataReplicator.ProcessReceivedData(*PartyData);
+		if (!bHasReceivedRepData)
+		{
+			bHasReceivedRepData = true;
+			TryFinishInitialization();
+		}
 	}
 }
 
diff --git a/Engine/Plugins/Online/OnlineFramework/Source/Party/Public/Party/SocialParty.h b/Engine/Plugins/Online/OnlineFramework/Source/Party/Public/Party/SocialParty.h
--- a/Engine/Plugins/Online/OnlineFramework/Source/Party/Public/Party/SocialParty.h
+++ b/Engine/Plugins/Online/OnlineFramework/Source/Party/Public/Party/SocialParty.h
@@ -382,6 +382,7 @@ private:
 
 	bool bIsLeavingParty = false;
 	bool bIsInitialized = false;
+	bool bHasReceivedRepData = false;
 	TOptional<bool> bIsRequestingShutdown;
 
 	mutable FLeavePartyEvent OnPartyLeaveBeginEvent;
