Sort user by playing this game

diff --git a/Engine/Plugins/Online/OnlineFramework/Source/Party/Private/User/SocialUserList.cpp b/Engine/Plugins/Online/OnlineFramework/Source/Party/Private/User/SocialUserList.cpp
--- a/Engine/Plugins/Online/OnlineFramework/Source/Party/Private/User/SocialUserList.cpp
+++ b/Engine/Plugins/Online/OnlineFramework/Source/Party/Private/User/SocialUserList.cpp
@@ -377,7 +377,7 @@ struct FUserSortData
 
 	bool operator<(const FUserSortData& OtherSortData) const
 	{
-		// Goes from if online, then alphabetical
+		// Goes from if online, playing this game, then alphabetical
 		if (OnlineStatus == OtherSortData.OnlineStatus)
 		{
 			if (PlayingThisGame == OtherSortData.PlayingThisGame)
@@ -386,7 +386,7 @@ struct FUserSortData
 			}
 			else
 			{
-				return false;
+				return PlayingThisGame > OtherSortData.PlayingThisGame;
 			}
 		}
 		else
