Add option to always send invitations on the primary party system in addition to platform invitations
Fix crash when party system is shutdown

diff --git a/Engine/Plugins/Online/OnlineFramework/Source/Party/Private/Party/SocialParty.cpp b/Engine/Plugins/Online/OnlineFramework/Source/Party/Private/Party/SocialParty.cpp
--- a/Engine/Plugins/Online/OnlineFramework/Source/Party/Private/Party/SocialParty.cpp
+++ b/Engine/Plugins/Online/OnlineFramework/Source/Party/Private/Party/SocialParty.cpp
@@ -300,11 +300,10 @@ bool USocialParty::CanInviteUserInternal(const USocialUser& User) const
 bool USocialParty::TryInviteUser(const USocialUser& UserToInvite)
 {
 	bool bSentInvite = false;
-	ESocialSubsystem InvitationSubsystemType = ESocialSubsystem::MAX;
-
 	if (CanInviteUser(UserToInvite))
 	{
-		static const bool bPreferPlatformInvite = USocialSettings::ShouldPreferPlatformInvites();
+		const bool bPreferPlatformInvite = USocialSettings::ShouldPreferPlatformInvites();
+		const bool bMustSendPrimaryInvite = USocialSettings::MustSendPrimaryInvites();
 
 		const FUniqueNetIdRepl UserPrimaryId = UserToInvite.GetUserId(ESocialSubsystem::Primary);
 		const FUniqueNetIdRepl UserPlatformId = UserToInvite.GetUserId(ESocialSubsystem::Platform);
@@ -316,28 +315,30 @@ bool USocialParty::TryInviteUser(const USocialUser& UserToInvite)
 
 		if ((UserPlatformId.IsValid() && bIsOnlineOnPlatform) && (!UserPrimaryId.IsValid() || bPreferPlatformInvite))
 		{
-			InvitationSubsystemType = ESocialSubsystem::Platform;
-
 			// Platform invites are sent as session invites on platform OSS' - this way we get the OS popups one would expect on XBox, PS4, etc.
+			bool bSentPlatformInvite = false;
 			const IOnlineSessionPtr PlatformSessionInterface = Online::GetSessionInterface(GetWorld(), USocialManager::GetSocialOssName(ESocialSubsystem::Platform));
 			if (PlatformSessionInterface.IsValid())
 			{
 				//@todo DanH Party: Any way to know if the session invite was a success? If we don't know we can't show it :/ #future
-				bSentInvite = PlatformSessionInterface->SendSessionInviteToFriend(*GetOwningLocalMember().GetRepData().GetPlatformUniqueId(), NAME_PartySession, *UserPlatformId);
+				bSentPlatformInvite = PlatformSessionInterface->SendSessionInviteToFriend(*GetOwningLocalMember().GetRepData().GetPlatformUniqueId(), NAME_PartySession, *UserPlatformId);
 			}
+			OnInviteSentInternal(ESocialSubsystem::Platform, UserToInvite, bSentPlatformInvite);
+			bSentInvite |= bSentPlatformInvite;
 		}
-		else if (UserPrimaryId.IsValid())
+		if ((!bSentInvite || bMustSendPrimaryInvite) && UserPrimaryId.IsValid())
 		{
-			InvitationSubsystemType = ESocialSubsystem::Primary;
-
 			// Primary subsystem invites can be sent directly to the user via the party interface
 			const IOnlinePartyPtr PartyInterface = Online::GetPartyInterfaceChecked(GetWorld());
-			bSentInvite = PartyInterface->SendInvitation(*OwningLocalUserId, GetPartyId(), *UserPrimaryId);
+			const bool bSentPrimaryInvite = PartyInterface->SendInvitation(*OwningLocalUserId, GetPartyId(), *UserPrimaryId);
+			OnInviteSentInternal(ESocialSubsystem::Primary, UserToInvite, bSentPrimaryInvite);
+			bSentInvite |= bSentPrimaryInvite;
 		}
 	}
-
-	OnInviteSentInternal(InvitationSubsystemType, UserToInvite, bSentInvite);
-
+	else
+	{
+		OnInviteSentInternal(ESocialSubsystem::MAX, UserToInvite, false);
+	}
 	return bSentInvite;
 }
 
@@ -952,7 +953,8 @@ void USocialParty::HandlePartySystemStateChange(EPartySystemState NewState)
 		//set timer to turn this off in a minute?
 		UWorld* World = GetWorld();
 		FTimerHandle DummyHandle;
-		World->GetTimerManager().SetTimer(DummyHandle, FTimerDelegate::CreateLambda([this]() { SetIsRequestingShutdown(false); }), 60.0f, false);
+		TWeakObjectPtr<USocialParty> WeakThis(this);
+		World->GetTimerManager().SetTimer(DummyHandle, FTimerDelegate::CreateLambda([WeakThis, this]() { if (WeakThis.IsValid()) { SetIsRequestingShutdown(false); } }), 60.0f, false);
 	}
 }
 
diff --git a/Engine/Plugins/Online/OnlineFramework/Source/Party/Private/SocialSettings.cpp b/Engine/Plugins/Online/OnlineFramework/Source/Party/Private/SocialSettings.cpp
--- a/Engine/Plugins/Online/OnlineFramework/Source/Party/Private/SocialSettings.cpp
+++ b/Engine/Plugins/Online/OnlineFramework/Source/Party/Private/SocialSettings.cpp
@@ -37,6 +37,12 @@ bool USocialSettings::ShouldPreferPlatformInvites()
 	return SettingsCDO.bPreferPlatformInvites;
 }
 
+bool USocialSettings::MustSendPrimaryInvites()
+{
+	const USocialSettings& SettingsCDO = *GetDefault<USocialSettings>();
+	return SettingsCDO.bMustSendPrimaryInvites;
+}
+
 int32 USocialSettings::GetDefaultMaxPartySize()
 {
 #if !UE_BUILD_SHIPPING
diff --git a/Engine/Plugins/Online/OnlineFramework/Source/Party/Public/SocialSettings.h b/Engine/Plugins/Online/OnlineFramework/Source/Party/Public/SocialSettings.h
--- a/Engine/Plugins/Online/OnlineFramework/Source/Party/Public/SocialSettings.h
+++ b/Engine/Plugins/Online/OnlineFramework/Source/Party/Public/SocialSettings.h
@@ -22,6 +22,7 @@ public:
 
 	static FString GetUniqueIdEnvironmentPrefix(ESocialSubsystem SubsystemType);
 	static bool ShouldPreferPlatformInvites();
+	static bool MustSendPrimaryInvites();
 	static int32 GetDefaultMaxPartySize();
 	static float GetUserListAutoUpdateRate();
 
@@ -42,6 +43,9 @@ private:
 	UPROPERTY(config)
 	bool bPreferPlatformInvites = true;
 
+	UPROPERTY(config)
+	bool bMustSendPrimaryInvites = false;
+
 	UPROPERTY(config)
 	float UserListAutoUpdateRate = 7.5f;
 };
\ No newline at end of file
