Re-enabled "SameParty" presence filtering for SocialUserList.

Added a new type of SocialRelationship "Any" which includes all users except the toolkit's owning user.

Added an array of GameSpecificStatus filters to FSocialUserListConfig that allows additional filtering whenever a user's OnGameSpecificStatusChanged event updates, which leads to...

Added an OnGameSpecificStatusChanged event to SocialUser, which can be used by games that have additional game-specific filtering needs that they only want evaulated at specific times (instead of every time a SocialUser's presence changes)

diff --git a/Engine/Plugins/Online/OnlineFramework/Source/Party/Private/Party/SocialParty.cpp b/Engine/Plugins/Online/OnlineFramework/Source/Party/Private/Party/SocialParty.cpp
--- a/Engine/Plugins/Online/OnlineFramework/Source/Party/Private/Party/SocialParty.cpp
+++ b/Engine/Plugins/Online/OnlineFramework/Source/Party/Private/Party/SocialParty.cpp
@@ -1135,6 +1135,19 @@ void USocialParty::LeaveParty(const FOnLeavePartyAttemptComplete& OnLeaveAttempt
 	}
 }
 
+bool USocialParty::ContainsUser(const USocialUser& User) const
+{
+	for(const UPartyMember* PartyMember : GetPartyMembers())
+	{
+		if (&PartyMember->GetSocialUser() == &User)
+		{
+			return true;
+		}
+	}
+
+	return false;
+}
+
 ULocalPlayer& USocialParty::GetOwningLocalPlayer() const
 {
 	//@todo DanH Party: This is a wee bit heavy - should be able to do this in fewer steps
diff --git a/Engine/Plugins/Online/OnlineFramework/Source/Party/Private/SocialDebugTools.cpp b/Engine/Plugins/Online/OnlineFramework/Source/Party/Private/SocialDebugTools.cpp
--- a/Engine/Plugins/Online/OnlineFramework/Source/Party/Private/SocialDebugTools.cpp
+++ b/Engine/Plugins/Online/OnlineFramework/Source/Party/Private/SocialDebugTools.cpp
@@ -633,6 +633,7 @@ void USocialDebugTools::FInstanceContext::Init()
 		if (OnlineParty.IsValid())
 		{
 			PartyInviteReceivedDelegateHandle = OnlineParty->AddOnPartyInviteReceivedDelegate_Handle(FOnPartyInviteReceivedDelegate::CreateUObject(&Owner, &USocialDebugTools::HandlePartyInviteReceived));
+			PartyJoinRequestReceivedDelegateHandle = OnlineParty->AddOnPartyJoinRequestReceivedDelegate_Handle(FOnPartyJoinRequestReceivedDelegate::CreateUObject(&Owner, &USocialDebugTools::HandlePartyJoinRequestReceived));
 		}
 	}
 }
@@ -652,6 +653,7 @@ void USocialDebugTools::FInstanceContext::Shutdown()
 		if (OnlineParty.IsValid())
 		{
 			OnlineParty->ClearOnPartyInviteReceivedDelegate_Handle(PartyInviteReceivedDelegateHandle);
+			OnlineParty->ClearOnPartyJoinRequestReceivedDelegate_Handle(PartyJoinRequestReceivedDelegateHandle);
 			OnlineParty.Reset();
 		}
 		OnlineSub->Shutdown();
@@ -718,4 +720,17 @@ void USocialDebugTools::HandlePartyInviteReceived(const FUniqueNetId& LocalUserI
 	}
 }
 
+void USocialDebugTools::HandlePartyJoinRequestReceived(const FUniqueNetId& LocalUserId, const FOnlinePartyId& PartyId, const FUniqueNetId& SenderId, const FString& Platform, const FOnlinePartyData& PartyData)
+{
+	FInstanceContext* Context = GetContextForUser(LocalUserId);
+	if (Context)
+	{
+		IOnlinePartyPtr OnlineParty = Context->GetOSS()->GetPartyInterface();
+		if (OnlineParty.IsValid())
+		{
+			OnlineParty->ApproveJoinRequest(LocalUserId, PartyId, SenderId, true);
+			return;
+		}
+	}
+}
 
diff --git a/Engine/Plugins/Online/OnlineFramework/Source/Party/Private/User/SocialUser.cpp b/Engine/Plugins/Online/OnlineFramework/Source/Party/Private/User/SocialUser.cpp
--- a/Engine/Plugins/Online/OnlineFramework/Source/Party/Private/User/SocialUser.cpp
+++ b/Engine/Plugins/Online/OnlineFramework/Source/Party/Private/User/SocialUser.cpp
@@ -228,6 +228,12 @@ TArray<ESocialSubsystem> USocialUser::GetRelationshipSubsystems(ESocialRelations
 		{
 			switch (Relationship)
 			{
+			case ESocialRelationship::Any:
+				if (!IsLocalUser())
+				{
+					RelationshipSubsystems.Add(SubsystemInfoPair.Key);
+				}
+				break;
 			case ESocialRelationship::FriendInviteReceived:
 				if (SubsystemInfoPair.Value.GetFriendInviteStatus() == EInviteStatus::PendingInbound)
 				{
diff --git a/Engine/Plugins/Online/OnlineFramework/Source/Party/Private/User/SocialUserList.cpp b/Engine/Plugins/Online/OnlineFramework/Source/Party/Private/User/SocialUserList.cpp
--- a/Engine/Plugins/Online/OnlineFramework/Source/Party/Private/User/SocialUserList.cpp
+++ b/Engine/Plugins/Online/OnlineFramework/Source/Party/Private/User/SocialUserList.cpp
@@ -9,6 +9,8 @@
 #include "Interfaces/OnlinePresenceInterface.h"
 #include "SocialSettings.h"
 #include "Algo/Transform.h"
+#include "SocialManager.h"
+#include "Party/SocialParty.h"
 
 TSharedRef<FSocialUserList> FSocialUserList::CreateUserList(USocialToolkit& InOwnerToolkit, const FSocialUserListConfig& InConfig)
 {
@@ -57,13 +59,24 @@ void FSocialUserList::InitializeList()
 	OwnerToolkit->OnUserBlocked().AddSP(this, &FSocialUserList::HandleUserBlocked);
 
 	// Run through all the users on the toolkit and add any that qualify for this list
+	check(Users.Num() == 0);
 	for (USocialUser* User : OwnerToolkit->GetAllUsers())
 	{
 		check(User);
 		TryAddUserFast(*User);
 	}
+	
+	if (EnumHasAnyFlags(ListConfig.ForbiddenPresenceFlags, ESocialUserStateFlags::SameParty) ||
+		EnumHasAnyFlags(ListConfig.RequiredPresenceFlags, ESocialUserStateFlags::SameParty))
+	{
+		OwnerToolkit->GetSocialManager().OnPartyJoined().AddSP(this, &FSocialUserList::HandlePartyJoined);
+		if (USocialParty* PersistentParty = OwnerToolkit->GetSocialManager().GetPersistentParty())
+		{
+			HandlePartyJoined(*PersistentParty);
+		}
+	}
 
-	if (ListConfig.bRequireAutoUpdate)
+	if (ListConfig.bAutoUpdate)
 	{
 		SetAutoUpdatePeriod(USocialSettings::GetUserListAutoUpdateRate());
 	}
@@ -201,9 +214,21 @@ void FSocialUserList::HandleRecentPlayerRemoved(USocialUser& RemovedUser, ESocia
 }
 
 void FSocialUserList::HandleUserPresenceChanged(ESocialSubsystem SubsystemType, USocialUser* User)
+{
+	MarkUserAsDirty(*User);
+}
+
+void FSocialUserList::HandleUserGameSpecificStatusChanged(USocialUser* User)
+{
+	// passing dummy Subsystem because HandleUserPresence changed currently doesn't care
+	MarkUserAsDirty(*User);
+	UpdateNow();
+}
+
+void FSocialUserList::MarkUserAsDirty(USocialUser& User)
 {
 	// Save this dirtied user for re-evaluation during the next update
-	UsersWithDirtyPresence.Add(User);
+	UsersWithDirtyPresence.Add(&User);
 	bNeedsSort = true;
 }
 
@@ -241,11 +266,29 @@ void FSocialUserList::TryAddUserFast(USocialUser& User)
 				User.OnUserPresenceChanged().AddSP(this, &FSocialUserList::HandleUserPresenceChanged, &User);
 			}
 
+			if (ListConfig.GameSpecificStatusFilters.Num() > 0)
+			{
+				if (!User.OnUserGameSpecificStatusChanged().IsBoundToObject(this))
+				{
+					User.OnUserGameSpecificStatusChanged().AddSP(this, &FSocialUserList::HandleUserGameSpecificStatusChanged, &User);
+				}
+			}
+
 			// Check that the user's current presence is acceptable
 			if (EvaluateUserPresence(User, RelationshipSubsystem))
 			{
 				// Last step is to check the custom filter, if provided
 				bCanAdd = ListConfig.OnCustomFilterUser.IsBound() ? ListConfig.OnCustomFilterUser.Execute(User) : true;
+
+				// do an initial pass on the GameSpecificStatusFilters (these will only be run again when the user broadcasts OnGameSpecificStatusChanged)
+				for (TFunction<bool(const USocialUser&)> CustomFilterFunction : ListConfig.GameSpecificStatusFilters)
+				{
+					if (!bCanAdd)
+					{
+						break;
+					}
+					bCanAdd &= CustomFilterFunction(User);
+				}
 			}
 		}
 	}
@@ -304,6 +347,15 @@ void FSocialUserList::TryRemoveUserFast(USocialUser& User)
 			{
 				// We're going to keep the user based on the stock filters, but the custom filter can still veto
 				bRemoveUser = ListConfig.OnCustomFilterUser.IsBound() ? !ListConfig.OnCustomFilterUser.Execute(User) : false;
+
+				// do an initial pass on the GameSpecificStatusFilters (these will only be run again when the user broadcasts OnGameSpecificStatusChanged)
+				if (!bRemoveUser)
+				{
+					for (TFunction<bool(const USocialUser&)> CustomFilterFunction : ListConfig.GameSpecificStatusFilters)
+					{
+						bRemoveUser |= !CustomFilterFunction(User);
+					}
+				}
 			}
 		}
 	}
@@ -325,6 +377,7 @@ void FSocialUserList::TryRemoveUserFast(USocialUser& User)
 		{
 			// Not only does this user not qualify for the list, they don't even have the appropriate relationship anymore (so we no longer care about presence changes)
 			User.OnUserPresenceChanged().RemoveAll(this);
+			User.OnUserGameSpecificStatusChanged().RemoveAll(this);
 		}
 	}
 }
@@ -335,11 +388,20 @@ bool FSocialUserList::EvaluateUserPresence(const USocialUser& User, ESocialSubsy
 	{
 		if (const FOnlineUserPresence* UserPresence = User.GetFriendPresenceInfo(SubsystemType))
 		{
+			bool bInSameParty = false;
+			if (OwnerToolkit.IsValid())
+			{
+				if (const USocialParty* CurrentParty = OwnerToolkit->GetSocialManager().GetPersistentParty())
+				{
+					bInSameParty = CurrentParty->ContainsUser(User);
+				}
+			}
+
 			return EvaluatePresenceFlag(UserPresence->bIsOnline, ESocialUserStateFlags::Online)
-				&& EvaluatePresenceFlag(UserPresence->bIsPlayingThisGame, ESocialUserStateFlags::SameApp);
+				&& EvaluatePresenceFlag(UserPresence->bIsPlayingThisGame, ESocialUserStateFlags::SameApp)
+				&& EvaluatePresenceFlag(bInSameParty, ESocialUserStateFlags::SameParty);
 			// && EvaluatePresenceFlag(UserPresence->bIsJoinable, ESocialUserStateFlags::Joinable) <-- //@todo DanH: This property exists on presence, but is ALWAYS false... 
 			// && EvaluateFlag(UserPresence->?, ESocialUserStateFlag::SamePlatform)
-			// && EvaluateFlag(UserPresence->?, ESocialUserStateFlag::SameParty)
 			// && EvaluateFlag(UserPresence->?, ESocialUserStateFlag::LookingForGroup)
 		}
 		return false;
@@ -408,6 +470,12 @@ bool FSocialUserList::HandleAutoUpdateList(float)
 
 void FSocialUserList::UpdateListInternal()
 {
+	if (NumIgnoreUpdateRequests > 0)
+	{
+		return;
+	}
+
+	
 	// Re-evaluate whether each user with dirtied presence is still fit for the list
 	for (TWeakObjectPtr<USocialUser> DirtyUser : UsersWithDirtyPresence)
 	{
@@ -485,3 +553,31 @@ void FSocialUserList::UpdateListInternal()
 		OnUpdateComplete().Broadcast();
 	}
 }
+
+void FSocialUserList::HandlePartyJoined(USocialParty& Party)
+{
+	Party.OnPartyMemberCreated().AddSP(this, &FSocialUserList::HandlePartyMemberCreated);
+
+	// ignore updates while we populate the party members (lots of redundant actions)
+	NumIgnoreUpdateRequests++;
+	for (UPartyMember* PartyMember : Party.GetPartyMembers())
+	{
+		HandlePartyMemberCreated(*PartyMember);
+	}
+	NumIgnoreUpdateRequests--;
+
+	UpdateNow();
+}
+
+void FSocialUserList::HandlePartyMemberCreated(UPartyMember& Member)
+{
+	Member.OnLeftParty().AddSP(this, &FSocialUserList::HandlePartyMemberLeft, &Member);
+	MarkUserAsDirty(Member.GetSocialUser());
+	UpdateNow();
+}
+
+void FSocialUserList::HandlePartyMemberLeft(EMemberExitedReason Reason, UPartyMember* Member)
+{
+	MarkUserAsDirty(Member->GetSocialUser());
+	UpdateNow();
+}
diff --git a/Engine/Plugins/Online/OnlineFramework/Source/Party/Private/User/SocialUserList.h b/Engine/Plugins/Online/OnlineFramework/Source/Party/Private/User/SocialUserList.h
--- a/Engine/Plugins/Online/OnlineFramework/Source/Party/Private/User/SocialUserList.h
+++ b/Engine/Plugins/Online/OnlineFramework/Source/Party/Private/User/SocialUserList.h
@@ -5,6 +5,8 @@
 #include "User/ISocialUserList.h"
 #include "UObject/GCObject.h"
 
+enum class EMemberExitedReason;
+
 class FSocialUserList : public ISocialUserList, public FGCObject, public TSharedFromThis<FSocialUserList>
 {
 public:
@@ -42,6 +44,9 @@ private:
 	void HandleRecentPlayerRemoved(USocialUser& RemovedUser, ESocialSubsystem SubsystemType);
 	
 	void HandleUserPresenceChanged(ESocialSubsystem SubsystemType, USocialUser* User);
+	void HandleUserGameSpecificStatusChanged(USocialUser* User);
+
+	void MarkUserAsDirty(USocialUser& User);
 
 	void TryAddUser(USocialUser& User);
 	void TryAddUserFast(USocialUser& User);
@@ -55,6 +60,10 @@ private:
 	bool HandleAutoUpdateList(float);
 	void UpdateListInternal();
 
+	void HandlePartyJoined(USocialParty& Party);
+	void HandlePartyMemberCreated(UPartyMember& Member);
+	void HandlePartyMemberLeft(EMemberExitedReason Reason, UPartyMember* Member);
+
 private:
 	FSocialUserList(USocialToolkit& InOwnerToolkit, const FSocialUserListConfig& Config);
 	void InitializeList();
@@ -73,6 +82,7 @@ private:
 	FSocialUserListConfig ListConfig;
 
 	// give external access to disable list update for perf
+	int32 NumIgnoreUpdateRequests = 0;
 	bool bAllowAutoUpdate = true;
 	bool bNeedsSort = false;
 	float AutoUpdatePeriod = 5.f;
diff --git a/Engine/Plugins/Online/OnlineFramework/Source/Party/Public/Party/SocialParty.h b/Engine/Plugins/Online/OnlineFramework/Source/Party/Public/Party/SocialParty.h
--- a/Engine/Plugins/Online/OnlineFramework/Source/Party/Public/Party/SocialParty.h
+++ b/Engine/Plugins/Online/OnlineFramework/Source/Party/Public/Party/SocialParty.h
@@ -114,6 +114,8 @@ public:
 	{
 		return Cast<MemberT>(GetMemberInternal(MemberId));
 	}
+
+	bool ContainsUser(const USocialUser& User) const;
 	
 	ULocalPlayer& GetOwningLocalPlayer() const;
 	const FUniqueNetIdRepl& GetOwningLocalUserId() const { return OwningLocalUserId; }
diff --git a/Engine/Plugins/Online/OnlineFramework/Source/Party/Public/SocialDebugTools.h b/Engine/Plugins/Online/OnlineFramework/Source/Party/Public/SocialDebugTools.h
--- a/Engine/Plugins/Online/OnlineFramework/Source/Party/Public/SocialDebugTools.h
+++ b/Engine/Plugins/Online/OnlineFramework/Source/Party/Public/SocialDebugTools.h
@@ -71,6 +71,7 @@ public:
 		FDelegateHandle PresenceReceivedDelegateHandle;
 		FDelegateHandle FriendInviteReceivedDelegateHandle;
 		FDelegateHandle PartyInviteReceivedDelegateHandle;
+		FDelegateHandle PartyJoinRequestReceivedDelegateHandle;
 	};
 
 	FInstanceContext& GetContext(const FString& Instance);
@@ -95,4 +96,5 @@ private:
 	// OSS callback handlers
 	void HandleFriendInviteReceived(const FUniqueNetId& LocalUserId, const FUniqueNetId& FriendId);
 	void HandlePartyInviteReceived(const FUniqueNetId& LocalUserId, const FOnlinePartyId& PartyId, const FUniqueNetId& SenderId);
-};
+	void HandlePartyJoinRequestReceived(const FUniqueNetId& LocalUserId, const FOnlinePartyId& PartyId, const FUniqueNetId& SenderId, const FString& Platform, const FOnlinePartyData& PartyData);
+};
\ No newline at end of file
diff --git a/Engine/Plugins/Online/OnlineFramework/Source/Party/Public/SocialTypes.h b/Engine/Plugins/Online/OnlineFramework/Source/Party/Public/SocialTypes.h
--- a/Engine/Plugins/Online/OnlineFramework/Source/Party/Public/SocialTypes.h
+++ b/Engine/Plugins/Online/OnlineFramework/Source/Party/Public/SocialTypes.h
@@ -47,6 +47,7 @@ inline const TCHAR* LexToString(ESocialSubsystem InSubsystem)
 UENUM()
 enum class ESocialRelationship : uint8
 {
+	Any,
 	FriendInviteReceived,
 	FriendInviteSent,
 	PartyInvite,
diff --git a/Engine/Plugins/Online/OnlineFramework/Source/Party/Public/User/ISocialUserList.h b/Engine/Plugins/Online/OnlineFramework/Source/Party/Public/User/ISocialUserList.h
--- a/Engine/Plugins/Online/OnlineFramework/Source/Party/Public/User/ISocialUserList.h
+++ b/Engine/Plugins/Online/OnlineFramework/Source/Party/Public/User/ISocialUserList.h
@@ -34,8 +34,15 @@ public:
 	TArray<ESocialSubsystem> ForbiddenSubsystems;
 	ESocialUserStateFlags RequiredPresenceFlags = ESocialUserStateFlags::None;
 	ESocialUserStateFlags ForbiddenPresenceFlags = ESocialUserStateFlags::None;
+
+	// Delegate that is evaluated every time we test a user for inclusion/exclusion
 	FOnCustomFilterUser OnCustomFilterUser;
-	bool bRequireAutoUpdate = false;
+
+	// These functions run whenever the User broadcasts OnGameSpecificStatusChanged and if they all return true will trigger a full re-evaluation of this user for list eligibility
+	TArray<TFunction<bool(const USocialUser&)>> GameSpecificStatusFilters;
+
+	// Whether or not the list should be polled regularly for updates (as opposed to manually having UpdateNow triggered)
+	bool bAutoUpdate = false;
 };
 
 class ISocialUserList
diff --git a/Engine/Plugins/Online/OnlineFramework/Source/Party/Public/User/SocialUser.h b/Engine/Plugins/Online/OnlineFramework/Source/Party/Public/User/SocialUser.h
--- a/Engine/Plugins/Online/OnlineFramework/Source/Party/Public/User/SocialUser.h
+++ b/Engine/Plugins/Online/OnlineFramework/Source/Party/Public/User/SocialUser.h
@@ -98,6 +98,11 @@ public:
 	DECLARE_EVENT_OneParam(USocialUser, FOnUserPresenceChanged, ESocialSubsystem)
 	FOnUserPresenceChanged& OnUserPresenceChanged() const { return OnUserPresenceChangedEvent; }
 
+	// provided so that lists with custom game-specific filtering (and any other listeners) can potentially re-evaluate a user
+	// the pattern here is similar to OnUserPresenceChanged but not subsystem-specific
+	DECLARE_EVENT(USocialUser, FOnUserGameSpecificStatusChanged)
+	FOnUserGameSpecificStatusChanged& OnUserGameSpecificStatusChanged() const { return OnUserGameSpecificStatusChangedEvent; }
+
 	DECLARE_EVENT_OneParam(USocialUser, FOnFriendRemoved, ESocialSubsystem)
 	FOnFriendRemoved& OnFriendRemoved() const { return OnFriendRemovedEvent; }
 	FOnFriendRemoved& OnFriendInviteRemoved() const { return OnFriendInviteRemovedEvent; }
@@ -178,4 +183,5 @@ private:
 	mutable FOnFriendRemoved OnFriendInviteRemovedEvent;
 	mutable FOnBlockedStatusChanged OnBlockedStatusChangedEvent;
 	mutable FOnSubsystemIdEstablished OnSubsystemIdEstablishedEvent;
+	mutable FOnUserGameSpecificStatusChanged OnUserGameSpecificStatusChangedEvent;
 };
\ No newline at end of file
