Back out changelist 5840051. EdTest errors

diff --git a/Engine/Plugins/Online/OnlineFramework/Source/Lobby/Private/LobbyBeaconHost.cpp b/Engine/Plugins/Online/OnlineFramework/Source/Lobby/Private/LobbyBeaconHost.cpp
--- a/Engine/Plugins/Online/OnlineFramework/Source/Lobby/Private/LobbyBeaconHost.cpp
+++ b/Engine/Plugins/Online/OnlineFramework/Source/Lobby/Private/LobbyBeaconHost.cpp
@@ -218,9 +218,8 @@ void ALobbyBeaconHost::KickPlayer(ALobbyBeaconClient* ClientActor, const FText&
 	DisconnectClient(ClientActor);
 }
 
-bool ALobbyBeaconHost::ProcessJoinServer(ALobbyBeaconClient* ClientActor)
+void ALobbyBeaconHost::ProcessJoinServer(ALobbyBeaconClient* ClientActor)
 {
-	bool bSuccess = false;
 	if (LobbyState != nullptr)
 	{
 		ALobbyBeaconPlayerState* Player = LobbyState->GetPlayer(ClientActor);
@@ -228,8 +227,6 @@ bool ALobbyBeaconHost::ProcessJoinServer(ALobbyBeaconClient* ClientActor)
 		{
 			Player->bInLobby = false;
 			ClientActor->AckJoiningServer();
-
-			bSuccess = true;
 		}
 		else
 		{
@@ -241,8 +238,6 @@ bool ALobbyBeaconHost::ProcessJoinServer(ALobbyBeaconClient* ClientActor)
 	{
 		UE_LOG(LogLobbyBeacon, Warning, TEXT("No lobby state to handle user joining the game!"));
 	}
-
-	return bSuccess;
 }
 
 void ALobbyBeaconHost::ProcessDisconnect(ALobbyBeaconClient* ClientActor)
diff --git a/Engine/Plugins/Online/OnlineFramework/Source/Lobby/Public/LobbyBeaconHost.h b/Engine/Plugins/Online/OnlineFramework/Source/Lobby/Public/LobbyBeaconHost.h
--- a/Engine/Plugins/Online/OnlineFramework/Source/Lobby/Public/LobbyBeaconHost.h
+++ b/Engine/Plugins/Online/OnlineFramework/Source/Lobby/Public/LobbyBeaconHost.h
@@ -125,7 +125,7 @@ protected:
 	 *
 	 * @param ClientActor client that is making the request
 	 */
-	virtual bool ProcessJoinServer(ALobbyBeaconClient* ClientActor);
+	void ProcessJoinServer(ALobbyBeaconClient* ClientActor);
 
 	/**
 	 * Handle a request to disconnect a given client from the lobby beacon
