Add IOnlinePartySystem::CleanupParties to cleanup party memberships for a user

diff --git a/Engine/Plugins/Online/OnlineFramework/Source/Party/Private/SocialDebugTools.cpp b/Engine/Plugins/Online/OnlineFramework/Source/Party/Private/SocialDebugTools.cpp
--- a/Engine/Plugins/Online/OnlineFramework/Source/Party/Private/SocialDebugTools.cpp
+++ b/Engine/Plugins/Online/OnlineFramework/Source/Party/Private/SocialDebugTools.cpp
@@ -319,6 +319,32 @@ void USocialDebugTools::LeaveParty(const FString& Instance, const FLeavePartyCom
 	OnComplete.ExecuteIfBound(false);
 }
 
+void USocialDebugTools::CleanupParties(const FString& Instance, const FCleanupPartiesComplete& OnComplete)
+{
+	IOnlineSubsystem* OnlineSub = GetContext(Instance).GetOSS();
+	if (OnlineSub)
+	{
+		IOnlineIdentityPtr OnlineIdentity = OnlineSub->GetIdentityInterface();
+		if (OnlineIdentity.IsValid())
+		{
+			TSharedPtr<const FUniqueNetId> UserId = OnlineIdentity->GetUniquePlayerId(LocalUserNum);
+			if (UserId.IsValid())
+			{
+				IOnlinePartyPtr OnlineParty = OnlineSub->GetPartyInterface();
+				if (OnlineParty.IsValid())
+				{
+					OnlineParty->CleanupParties(*UserId, FOnCleanupPartiesComplete::CreateLambda([this, Instance, OnComplete](const FUniqueNetId& LocalUserId, const FOnlineError& Result)
+					{
+						OnComplete.ExecuteIfBound(Result.WasSuccessful());
+					}));
+					return;
+				}
+			}
+		}
+	}
+	OnComplete.ExecuteIfBound(false);
+}
+
 void USocialDebugTools::SetPartyMemberData(const FString& Instance, const UStruct* StructType, const void* StructData, const FSetPartyMemberDataComplete& OnComplete)
 {
 	check(StructType);
@@ -548,13 +574,16 @@ bool USocialDebugTools::RunCommand(const TCHAR* Cmd, const TArray<FString>& Targ
 
 				if (bSuccess)
 				{
-					LeaveParty(TargetInstance, FLeavePartyComplete::CreateLambda([this, TargetInstance, FriendName](bool bLeavePartySuccess)
+					CleanupParties(TargetInstance, FCleanupPartiesComplete::CreateLambda([this, TargetInstance, FriendName](bool bCleanupPartiesSuccess)
 					{
-						UE_LOG(LogParty, Display, TEXT("Leave party OSS context[%s] %s"), *TargetInstance, *LexToString(bLeavePartySuccess));
-
-						JoinParty(TargetInstance, FriendName, FJoinPartyComplete::CreateLambda([this, TargetInstance](bool bJoinPartySuccess)
+						LeaveParty(TargetInstance, FLeavePartyComplete::CreateLambda([this, TargetInstance, FriendName](bool bLeavePartySuccess)
 						{
-							UE_LOG(LogParty, Display, TEXT("Join party OSS context[%s] %s"), *TargetInstance, *LexToString(bJoinPartySuccess));
+							UE_LOG(LogParty, Display, TEXT("Leave party OSS context[%s] %s"), *TargetInstance, *LexToString(bLeavePartySuccess));
+
+							JoinParty(TargetInstance, FriendName, FJoinPartyComplete::CreateLambda([this, TargetInstance](bool bJoinPartySuccess)
+							{
+								UE_LOG(LogParty, Display, TEXT("Join party OSS context[%s] %s"), *TargetInstance, *LexToString(bJoinPartySuccess));
+							}));
 						}));
 					}));
 				}
diff --git a/Engine/Plugins/Online/OnlineFramework/Source/Party/Public/SocialDebugTools.h b/Engine/Plugins/Online/OnlineFramework/Source/Party/Public/SocialDebugTools.h
--- a/Engine/Plugins/Online/OnlineFramework/Source/Party/Public/SocialDebugTools.h
+++ b/Engine/Plugins/Online/OnlineFramework/Source/Party/Public/SocialDebugTools.h
@@ -41,6 +41,9 @@ public:
 	DECLARE_DELEGATE_OneParam(FLeavePartyComplete, bool);
 	virtual void LeaveParty(const FString& Instance, const FLeavePartyComplete& OnComplete);
 
+	DECLARE_DELEGATE_OneParam(FCleanupPartiesComplete, bool);
+	virtual void CleanupParties(const FString& Instance, const FCleanupPartiesComplete& OnComplete);
+
 	DECLARE_DELEGATE_OneParam(FSetPartyMemberDataComplete, bool);
 	virtual void SetPartyMemberData(const FString& Instance, const UStruct* StructType, const void* StructData, const FSetPartyMemberDataComplete& OnComplete);
 	virtual void SetPartyMemberDataJson(const FString& Instance, const FString& JsonStr, const FSetPartyMemberDataComplete& OnComplete);
diff --git a/Engine/Plugins/Online/OnlineSubsystem/Source/Public/Interfaces/OnlinePartyInterface.h b/Engine/Plugins/Online/OnlineSubsystem/Source/Public/Interfaces/OnlinePartyInterface.h
--- a/Engine/Plugins/Online/OnlineSubsystem/Source/Public/Interfaces/OnlinePartyInterface.h
+++ b/Engine/Plugins/Online/OnlineSubsystem/Source/Public/Interfaces/OnlinePartyInterface.h
@@ -645,6 +645,13 @@ enum class EPartySystemState
  * @param Result Result of the operation
  */
 DECLARE_DELEGATE_TwoParams(FOnRestorePartiesComplete, const FUniqueNetId& /*LocalUserId*/, const FOnlineError& /*Result*/);
+/**
+ * Cleanup parties async task completed callback
+ *
+ * @param LocalUserId id of user that initiated the request
+ * @param Result Result of the operation
+ */
+DECLARE_DELEGATE_TwoParams(FOnCleanupPartiesComplete, const FUniqueNetId& /*LocalUserId*/, const FOnlineError& /*Result*/);
 /**
  * Party creation async task completed callback
  *
@@ -962,6 +969,14 @@ public:
 	 */
 	virtual void RestoreParties(const FUniqueNetId& LocalUserId, const FOnRestorePartiesComplete& CompletionDelegate) = 0;
 	
+	/**
+	 * Cleanup party state. This will cleanup the local party state and attempt to cleanup party memberships on an external service if possible.  Intended to be called for development purposes.
+	 *
+	 * @param LocalUserId the user to cleanup the parties for
+	 * @param CompletionDelegate the delegate to trigger on completion
+	 */
+	virtual void CleanupParties(const FUniqueNetId& LocalUserId, const FOnCleanupPartiesComplete& CompletionDelegate) = 0;
+	
 	/**
 	 * Create a new party
 	 *
