Modified EvaluateUserPresence check to not require actual Presence.  Without Presence a user is treated as Offline and Not Playing This Game, but can still be evaluated for In Party or not.  This ensures that lists using the SameParty presence feature can contain SocialUsers that aren't friends.

stephan.jiang, dan.hertzka
#tests Compile DebugGame Editor Win64 / three-way -game, verified party list and block list

diff --git a/Engine/Plugins/Online/OnlineFramework/Source/Party/Private/User/SocialUserList.cpp b/Engine/Plugins/Online/OnlineFramework/Source/Party/Private/User/SocialUserList.cpp
--- a/Engine/Plugins/Online/OnlineFramework/Source/Party/Private/User/SocialUserList.cpp
+++ b/Engine/Plugins/Online/OnlineFramework/Source/Party/Private/User/SocialUserList.cpp
@@ -24,6 +24,7 @@ FSocialUserList::FSocialUserList(USocialToolkit& InOwnerToolkit, const FSocialUs
 	, ListConfig(InConfig)
 {
 	if (HasPresenceFilters() &&
+		ListConfig.RelationshipType != ESocialRelationship::Any &&
 		ListConfig.RelationshipType != ESocialRelationship::Friend &&
 		ListConfig.RelationshipType != ESocialRelationship::PartyInvite)
 	{
@@ -386,26 +387,31 @@ bool FSocialUserList::EvaluateUserPresence(const USocialUser& User, ESocialSubsy
 {
 	if (HasPresenceFilters())
 	{
+		bool bIsOnline = false;
+		bool bIsPlayingThisGame = false;
+		bool bInSameParty = false;
 		if (const FOnlineUserPresence* UserPresence = User.GetFriendPresenceInfo(SubsystemType))
 		{
-			bool bInSameParty = false;
-			if (OwnerToolkit.IsValid())
+			bIsOnline = UserPresence->bIsOnline;
+			bIsPlayingThisGame = UserPresence->bIsPlayingThisGame;
+		}
+		
+		if (OwnerToolkit.IsValid())
+		{
+			if (const USocialParty* CurrentParty = OwnerToolkit->GetSocialManager().GetPersistentParty())
 			{
-				if (const USocialParty* CurrentParty = OwnerToolkit->GetSocialManager().GetPersistentParty())
-				{
-					bInSameParty = CurrentParty->ContainsUser(User);
-				}
+				bInSameParty = CurrentParty->ContainsUser(User);
 			}
-
-			return EvaluatePresenceFlag(UserPresence->bIsOnline, ESocialUserStateFlags::Online)
-				&& EvaluatePresenceFlag(UserPresence->bIsPlayingThisGame, ESocialUserStateFlags::SameApp)
-				&& EvaluatePresenceFlag(bInSameParty, ESocialUserStateFlags::SameParty);
-			// && EvaluatePresenceFlag(UserPresence->bIsJoinable, ESocialUserStateFlags::Joinable) <-- //@todo DanH: This property exists on presence, but is ALWAYS false... 
-			// && EvaluateFlag(UserPresence->?, ESocialUserStateFlag::SamePlatform)
-			// && EvaluateFlag(UserPresence->?, ESocialUserStateFlag::LookingForGroup)
 		}
-		return false;
+
+		return EvaluatePresenceFlag(bIsOnline, ESocialUserStateFlags::Online)
+			&& EvaluatePresenceFlag(bIsPlayingThisGame, ESocialUserStateFlags::SameApp)
+			&& EvaluatePresenceFlag(bInSameParty, ESocialUserStateFlags::SameParty);
+		// && EvaluatePresenceFlag(UserPresence->bIsJoinable, ESocialUserStateFlags::Joinable) <-- //@todo DanH: This property exists on presence, but is ALWAYS false... 
+		// && EvaluateFlag(UserPresence->?, ESocialUserStateFlag::SamePlatform)
+		// && EvaluateFlag(UserPresence->?, ESocialUserStateFlag::LookingForGroup)
 	}
+
 	return true;
 }
 
diff --git a/Engine/Plugins/Online/OnlineFramework/Source/Party/Public/SocialTypes.h b/Engine/Plugins/Online/OnlineFramework/Source/Party/Public/SocialTypes.h
--- a/Engine/Plugins/Online/OnlineFramework/Source/Party/Public/SocialTypes.h
+++ b/Engine/Plugins/Online/OnlineFramework/Source/Party/Public/SocialTypes.h
@@ -58,6 +58,31 @@ enum class ESocialRelationship : uint8
 	// Follower, (?)
 };
 
+inline const TCHAR* LexToString(ESocialRelationship Relationship)
+{
+	switch (Relationship)
+	{
+	case ESocialRelationship::Any:
+		return TEXT("Any");
+	case ESocialRelationship::FriendInviteReceived:
+		return TEXT("FriendInviteReceived");
+	case ESocialRelationship::FriendInviteSent:
+		return TEXT("FriendInviteSent");
+	case ESocialRelationship::PartyInvite:
+		return TEXT("PartyInvite");
+	case ESocialRelationship::Friend:
+		return TEXT("Friend");
+	case ESocialRelationship::BlockedPlayer:
+		return TEXT("BlockedPlayer");
+	case ESocialRelationship::SuggestedFriend:
+		return TEXT("SuggestedFriend");
+	case ESocialRelationship::RecentPlayer:
+		return TEXT("RecentPlayer");
+	}
+
+	return TEXT("Unknown");
+}
+
 UENUM()
 enum class ECrossplayPreference : uint8
 {
