Removed a speculative ensure() when performing removal of users from lists.  It is possible with correct timing to have someone in PendingRemove that was not previously in the Users list.  This is not currently considered problematic.

#tests Compile DebugGame Editor Win64

diff --git a/Engine/Plugins/Online/OnlineFramework/Source/Party/Private/User/SocialUserList.cpp b/Engine/Plugins/Online/OnlineFramework/Source/Party/Private/User/SocialUserList.cpp
--- a/Engine/Plugins/Online/OnlineFramework/Source/Party/Private/User/SocialUserList.cpp
+++ b/Engine/Plugins/Online/OnlineFramework/Source/Party/Private/User/SocialUserList.cpp
@@ -508,7 +508,7 @@ void FSocialUserList::UpdateListInternal()
 
 		Users.RemoveAllSwap(
 			[this] (USocialUser* User)
-			{
+		{
 				if (PendingRemovals.Contains(User))
 				{
 					PendingRemovals.Remove(User);
@@ -519,7 +519,6 @@ void FSocialUserList::UpdateListInternal()
 				return false;
 			});
 
-		ensure(PendingRemovals.Num() == 0);
 		PendingRemovals.Reset();
 	}
 
