Added a "Name" member to FSocialUserListConfig to aid in tracking / debugging inside the UserList functions when many lists are active.

#tests Compile DebugGame Editor Win64

diff --git a/Engine/Plugins/Online/OnlineFramework/Source/Party/Public/User/ISocialUserList.h b/Engine/Plugins/Online/OnlineFramework/Source/Party/Public/User/ISocialUserList.h
--- a/Engine/Plugins/Online/OnlineFramework/Source/Party/Public/User/ISocialUserList.h
+++ b/Engine/Plugins/Online/OnlineFramework/Source/Party/Public/User/ISocialUserList.h
@@ -29,6 +29,7 @@ class FSocialUserListConfig
 public:
 	FSocialUserListConfig() {}
 
+	FString Name = TEXT("");
 	ESocialRelationship RelationshipType = ESocialRelationship::Friend;
 	TArray<ESocialSubsystem> RelevantSubsystems;
 	TArray<ESocialSubsystem> ForbiddenSubsystems;
