Made USocialParty::LeaveParty virtual for game-specific overrides.

diff --git a/Engine/Plugins/Online/OnlineFramework/Source/Party/Public/Party/SocialParty.h b/Engine/Plugins/Online/OnlineFramework/Source/Party/Public/Party/SocialParty.h
--- a/Engine/Plugins/Online/OnlineFramework/Source/Party/Public/Party/SocialParty.h
+++ b/Engine/Plugins/Online/OnlineFramework/Source/Party/Public/Party/SocialParty.h
@@ -83,7 +83,7 @@ public:
 	void RefreshPublicJoinability();
 
 	DECLARE_DELEGATE_OneParam(FOnLeavePartyAttemptComplete, ELeavePartyCompletionResult)
-	void LeaveParty(const FOnLeavePartyAttemptComplete& OnLeaveAttemptComplete = FOnLeavePartyAttemptComplete());
+	virtual void LeaveParty(const FOnLeavePartyAttemptComplete& OnLeaveAttemptComplete = FOnLeavePartyAttemptComplete());
 
 	const FPartyRepData& GetRepData() const { return *PartyDataReplicator; }
 
