Addressed race condition causing stale platform session information to persist when a party member is promoted to leader as a result of the leader leaving

diff --git a/Engine/Plugins/Online/OnlineFramework/Source/Party/Private/Party/PartyTypes.cpp b/Engine/Plugins/Online/OnlineFramework/Source/Party/Private/Party/PartyTypes.cpp
--- a/Engine/Plugins/Online/OnlineFramework/Source/Party/Private/Party/PartyTypes.cpp
+++ b/Engine/Plugins/Online/OnlineFramework/Source/Party/Private/Party/PartyTypes.cpp
@@ -109,14 +109,27 @@ bool FJoinPartyResult::WasSuccessful() const
 // FOnlinePartyRepDataBase
 //////////////////////////////////////////////////////////////////////////
 
-void FOnlinePartyRepDataBase::LogPropertyChanged(const TCHAR* OwningStructTypeName, const TCHAR* ProperyName, bool bFromReplication) const
+void FOnlinePartyRepDataBase::LogSetPropertyFailure(const TCHAR* OwningStructTypeName, const TCHAR* PropertyName) const
+{
+	const USocialParty* OwningParty = GetOwnerParty();
+	if (ensure(OwningParty))
+	{
+		UE_LOG(LogParty, Warning, TEXT("Failed to modify RepData property [%s::%s] in party [%s] - local member [%s] does not have authority."),
+			OwningStructTypeName,
+			PropertyName,
+			*OwningParty->ToDebugString(),
+			*OwningParty->GetOwningLocalMember().ToDebugString(false));
+	}
+}
+
+void FOnlinePartyRepDataBase::LogPropertyChanged(const TCHAR* OwningStructTypeName, const TCHAR* PropertyName, bool bFromReplication) const
 {
 	const USocialParty* OwningParty = GetOwnerParty();
 	
 	// Only thing this lacks is the id of the party member for member rep data changes
 	UE_LOG(LogParty, VeryVerbose, TEXT("RepData property [%s::%s] changed %s in party [%s]"),
 		OwningStructTypeName,
-		ProperyName,
+		PropertyName,
 		bFromReplication ? TEXT("remotely") : TEXT("locally"),
 		ensure(OwningParty) ? *OwningParty->ToDebugString() : TEXT("unknown"));
 }
diff --git a/Engine/Plugins/Online/OnlineFramework/Source/Party/Private/Party/SocialParty.cpp b/Engine/Plugins/Online/OnlineFramework/Source/Party/Private/Party/SocialParty.cpp
--- a/Engine/Plugins/Online/OnlineFramework/Source/Party/Private/Party/SocialParty.cpp
+++ b/Engine/Plugins/Online/OnlineFramework/Source/Party/Private/Party/SocialParty.cpp
@@ -514,6 +514,19 @@ void USocialParty::OnLocalPlayerIsLeaderChanged(bool bIsLeader)
 
 		// Establish the privacy of the party to match the local player's preference
 		GetMutableRepData().SetPrivacySettings(GetDesiredPrivacySettings());
+
+		// It's possible that membership changes resulting in this promotion also require updates to the session info
+		//	If we found out about the changes in membership before learning we're the leader, we were unable to update the rep data accordingly
+		//	So, upon becoming leader, we must do a sweep to account for any such changes we missed out on
+		TArray<FName> AllMemberPlatformSubsystems;
+		for (UPartyMember* Member : GetPartyMembers())
+		{
+			AllMemberPlatformSubsystems.AddUnique(Member->GetPlatformOssName());
+		}
+		for (FName PlatformOssName : AllMemberPlatformSubsystems)
+		{
+			UpdatePlatformSessionLeader(PlatformOssName);
+		}
 	}
 	else
 	{
@@ -1530,6 +1543,11 @@ void USocialParty::FinalizePartyLeave(EMemberExitedReason Reason)
 
 void USocialParty::UpdatePlatformSessionLeader(FName PlatformOssName)
 {
+	if (!IsLocalPlayerPartyLeader())
+	{
+		return;
+	}
+
 	if (const FPartyPlatformSessionInfo* PlatformSessionInfo = GetRepData().FindSessionInfo(PlatformOssName))
 	{
 		UPartyMember* NewSessionOwner = nullptr;
diff --git a/Engine/Plugins/Online/OnlineFramework/Source/Party/Public/Party/PartyTypes.h b/Engine/Plugins/Online/OnlineFramework/Source/Party/Public/Party/PartyTypes.h
--- a/Engine/Plugins/Online/OnlineFramework/Source/Party/Public/Party/PartyTypes.h
+++ b/Engine/Plugins/Online/OnlineFramework/Source/Party/Public/Party/PartyTypes.h
@@ -408,7 +408,8 @@ protected:
 	virtual bool CanEditData() const { return false; }
 	virtual const USocialParty* GetOwnerParty() const { return nullptr; }
 
-	void LogPropertyChanged(const TCHAR* OwningStructTypeName, const TCHAR* ProperyName, bool bFromReplication) const;
+	void LogSetPropertyFailure(const TCHAR* OwningStructTypeName, const TCHAR* PropertyName) const;
+	void LogPropertyChanged(const TCHAR* OwningStructTypeName, const TCHAR* PropertyName, bool bFromReplication) const;
 
 	friend class FPartyDataReplicatorHelper;
 	template <typename> friend class TPartyDataReplicator;
@@ -459,21 +460,28 @@ EXPOSE_REP_DATA_PROPERTY_NO_SETTER(Owner, PropertyType, Property);	\
 SetterPrivacy:	\
 	void Set##Property(Property##ArgType New##Property)	\
 	{	\
-		if (CanEditData() && Property != New##Property)	\
+		if (CanEditData())	\
 		{	\
-			LogPropertyChanged(TEXT(#Owner), TEXT(#Property), false);	\
-			if (On##Property##ChangedDif().IsBound())	\
+			if (Property != New##Property)	\
 			{	\
-				PropertyType OldValue = Property;	\
-				Property = New##Property;	\
-				On##Property##ChangedDif().Broadcast(Property, OldValue);	\
+				LogPropertyChanged(TEXT(#Owner), TEXT(#Property), false);	\
+				if (On##Property##ChangedDif().IsBound())	\
+				{	\
+					PropertyType OldValue = Property;	\
+					Property = New##Property;	\
+					On##Property##ChangedDif().Broadcast(Property, OldValue);	\
+				}	\
+				else	\
+				{	\
+					Property = New##Property;	\
+				}	\
+				On##Property##Changed().Broadcast(Property);	\
+				OnDataChanged.ExecuteIfBound();	\
 			}	\
-			else	\
-			{	\
-				Property = New##Property;	\
-			}	\
-			On##Property##Changed().Broadcast(Property);	\
-			OnDataChanged.ExecuteIfBound();	\
+		}	\
+		else	\
+		{	\
+			LogSetPropertyFailure(TEXT(#Owner), TEXT(#Property));	\
 		}	\
 	}	\
 private: //
