Ignoring members that aren't on a platform with its own OSS when refreshing platform sessions after a promotion (wasn't affecting the result, but this is certainly more correct)

diff --git a/Engine/Plugins/Online/OnlineFramework/Source/Party/Private/Party/SocialParty.cpp b/Engine/Plugins/Online/OnlineFramework/Source/Party/Private/Party/SocialParty.cpp
--- a/Engine/Plugins/Online/OnlineFramework/Source/Party/Private/Party/SocialParty.cpp
+++ b/Engine/Plugins/Online/OnlineFramework/Source/Party/Private/Party/SocialParty.cpp
@@ -521,7 +521,11 @@ void USocialParty::OnLocalPlayerIsLeaderChanged(bool bIsLeader)
 		TArray<FName> AllMemberPlatformSubsystems;
 		for (UPartyMember* Member : GetPartyMembers())
 		{
-			AllMemberPlatformSubsystems.AddUnique(Member->GetPlatformOssName());
+			FName PlatformOssName = Member->GetPlatformOssName();
+			if (!PlatformOssName.IsNone())
+			{
+				AllMemberPlatformSubsystems.AddUnique(PlatformOssName);
+			}
 		}
 		for (FName PlatformOssName : AllMemberPlatformSubsystems)
 		{
