Adding warning log when attempting to update rep data without authority to do so
- Also making several root SocialUser actions virtual to allow game-specific behaviors/responses

diff --git a/Engine/Plugins/Online/OnlineFramework/Source/Party/Private/Party/PartyTypes.cpp b/Engine/Plugins/Online/OnlineFramework/Source/Party/Private/Party/PartyTypes.cpp
--- a/Engine/Plugins/Online/OnlineFramework/Source/Party/Private/Party/PartyTypes.cpp
+++ b/Engine/Plugins/Online/OnlineFramework/Source/Party/Private/Party/PartyTypes.cpp
@@ -112,14 +112,11 @@ bool FJoinPartyResult::WasSuccessful() const
 void FOnlinePartyRepDataBase::LogSetPropertyFailure(const TCHAR* OwningStructTypeName, const TCHAR* PropertyName) const
 {
 	const USocialParty* OwningParty = GetOwnerParty();
-	if (ensure(OwningParty))
-	{
-		UE_LOG(LogParty, Warning, TEXT("Failed to modify RepData property [%s::%s] in party [%s] - local member [%s] does not have authority."),
-			OwningStructTypeName,
-			PropertyName,
-			*OwningParty->ToDebugString(),
-			*OwningParty->GetOwningLocalMember().ToDebugString(false));
-	}
+	UE_LOG(LogParty, Warning, TEXT("Failed to modify RepData property [%s::%s] in party [%s] - local member [%s] does not have authority."),
+		OwningStructTypeName,
+		PropertyName,
+		OwningParty ? *OwningParty->ToDebugString() : TEXT("unknown"),
+		OwningParty ? *OwningParty->GetOwningLocalMember().ToDebugString(false) : TEXT("unknown"));
 }
 
 void FOnlinePartyRepDataBase::LogPropertyChanged(const TCHAR* OwningStructTypeName, const TCHAR* PropertyName, bool bFromReplication) const
@@ -131,5 +128,5 @@ void FOnlinePartyRepDataBase::LogPropertyChanged(const TCHAR* OwningStructTypeNa
 		OwningStructTypeName,
 		PropertyName,
 		bFromReplication ? TEXT("remotely") : TEXT("locally"),
-		ensure(OwningParty) ? *OwningParty->ToDebugString() : TEXT("unknown"));
+		OwningParty ? *OwningParty->ToDebugString() : TEXT("unknown"));
 }
diff --git a/Engine/Plugins/Online/OnlineFramework/Source/Party/Private/User/SocialUser.cpp b/Engine/Plugins/Online/OnlineFramework/Source/Party/Private/User/SocialUser.cpp
--- a/Engine/Plugins/Online/OnlineFramework/Source/Party/Private/User/SocialUser.cpp
+++ b/Engine/Plugins/Online/OnlineFramework/Source/Party/Private/User/SocialUser.cpp
@@ -571,6 +571,27 @@ void USocialUser::GetRichPresenceText(FText& OutRichPresence) const
 	}
 }
 
+bool USocialUser::IsRecentPlayer(ESocialSubsystem SubsystemType) const
+{
+	if (const FSubsystemUserInfo* SubsystemInfo = SubsystemInfoByType.Find(SubsystemType))
+	{
+		return SubsystemInfo->RecentPlayerInfo.IsValid();
+	}
+	return false;
+}
+
+bool USocialUser::IsRecentPlayer() const
+{
+	for (const TPair<ESocialSubsystem, FSubsystemUserInfo>& SubsystemInfoPair : SubsystemInfoByType)
+	{
+		if (SubsystemInfoPair.Value.RecentPlayerInfo.IsValid())
+		{
+			return true;
+		}
+	}
+	return false;
+}
+
 bool USocialUser::IsBlocked(ESocialSubsystem SubsystemType) const
 {
 	if (const FSubsystemUserInfo* SubsystemInfo = SubsystemInfoByType.Find(SubsystemType))
@@ -692,7 +713,7 @@ void USocialUser::JoinParty(const FOnlinePartyTypeId& PartyTypeId) const
 	{
 		IOnlinePartyPtr PartyInterface = Online::GetPartyInterfaceChecked(GetWorld());
 		PartyInterface->ClearInvitations(*GetOwningToolkit().GetLocalUserNetId(ESocialSubsystem::Primary), *GetUserId(ESocialSubsystem::Primary), nullptr);
-		OnPartyInviteAccepted().Broadcast();
+		OnPartyInviteAcceptedInternal(PartyTypeId);
 	}
 }
 
@@ -702,7 +723,7 @@ void USocialUser::RejectPartyInvite(const FOnlinePartyTypeId& PartyTypeId)
 	{
 		IOnlinePartyPtr PartyInterface = Online::GetPartyInterfaceChecked(GetWorld());
 		PartyInterface->RejectInvitation(*GetOwningToolkit().GetLocalUserNetId(ESocialSubsystem::Primary), *GetUserId(ESocialSubsystem::Primary));
-		OnPartyInviteRejected().Broadcast();
+		OnPartyInviteRejectedInternal(PartyTypeId);
 	}
 }
 
@@ -735,7 +756,7 @@ bool USocialUser::InviteToParty(const FOnlinePartyTypeId& PartyTypeId) const
 	return false;
 }
 
-bool USocialUser::BlockUser(ESocialSubsystem Subsystem)
+bool USocialUser::BlockUser(ESocialSubsystem Subsystem) const
 {
 	if (IOnlineSubsystem* Oss = GetOwningToolkit().GetSocialOss(Subsystem))
 	{
@@ -753,7 +774,7 @@ bool USocialUser::BlockUser(ESocialSubsystem Subsystem)
 	return false;
 }
 
-bool USocialUser::UnblockUser(ESocialSubsystem Subsystem)
+bool USocialUser::UnblockUser(ESocialSubsystem Subsystem) const
 {
 	if (IOnlineSubsystem* Oss = GetOwningToolkit().GetSocialOss(Subsystem))
 	{
@@ -1042,6 +1063,16 @@ void USocialUser::OnPresenceChangedInternal(ESocialSubsystem SubsystemType)
 	OnUserPresenceChanged().Broadcast(SubsystemType);
 }
 
+void USocialUser::OnPartyInviteAcceptedInternal(const FOnlinePartyTypeId& PartyTypeId) const
+{
+	OnPartyInviteAccepted().Broadcast();
+}
+
+void USocialUser::OnPartyInviteRejectedInternal(const FOnlinePartyTypeId& PartyTypeId) const
+{
+	OnPartyInviteRejected().Broadcast();
+}
+
 void USocialUser::NotifyPresenceChanged(ESocialSubsystem SubsystemType)
 {
 	OnPresenceChangedInternal(SubsystemType);
diff --git a/Engine/Plugins/Online/OnlineFramework/Source/Party/Public/User/SocialUser.h b/Engine/Plugins/Online/OnlineFramework/Source/Party/Public/User/SocialUser.h
--- a/Engine/Plugins/Online/OnlineFramework/Source/Party/Public/User/SocialUser.h
+++ b/Engine/Plugins/Online/OnlineFramework/Source/Party/Public/User/SocialUser.h
@@ -55,6 +55,9 @@ public:
 
 	virtual void GetRichPresenceText(FText& OutRichPresence) const;
 
+	bool IsRecentPlayer() const;
+	bool IsRecentPlayer(ESocialSubsystem SubsystemType) const;
+	
 	bool IsBlocked() const;
 	bool IsBlocked(ESocialSubsystem SubsystemType) const;
 
@@ -67,11 +70,11 @@ public:
 	bool HasAnyInteractionsAvailable() const;
 	TArray<FSocialInteractionHandle> GetAllAvailableInteractions() const;
 
-	bool CanSendFriendInvite(ESocialSubsystem SubsystemType) const;
-	bool SendFriendInvite(ESocialSubsystem SubsystemType);
-	bool AcceptFriendInvite(ESocialSubsystem SocialSubsystem) const;
-	bool RejectFriendInvite(ESocialSubsystem SocialSubsystem) const;
-	bool EndFriendship(ESocialSubsystem SocialSubsystem) const;
+	virtual bool CanSendFriendInvite(ESocialSubsystem SubsystemType) const;
+	virtual bool SendFriendInvite(ESocialSubsystem SubsystemType);
+	virtual bool AcceptFriendInvite(ESocialSubsystem SocialSubsystem) const;
+	virtual bool RejectFriendInvite(ESocialSubsystem SocialSubsystem) const;
+	virtual bool EndFriendship(ESocialSubsystem SocialSubsystem) const;
 
 	bool ShowPlatformProfile();
 
@@ -86,8 +89,8 @@ public:
 	bool CanInviteToParty(const FOnlinePartyTypeId& PartyTypeId) const;
 	bool InviteToParty(const FOnlinePartyTypeId& PartyTypeId) const;
 
-	bool BlockUser(ESocialSubsystem Subsystem);
-	bool UnblockUser(ESocialSubsystem Subsystem);
+	virtual bool BlockUser(ESocialSubsystem Subsystem) const;
+	virtual bool UnblockUser(ESocialSubsystem Subsystem) const;
 
 	UPartyMember* GetPartyMember(const FOnlinePartyTypeId& PartyTypeId) const;
 
@@ -132,6 +135,8 @@ PARTY_SCOPE:
 
 protected:
 	virtual void OnPresenceChangedInternal(ESocialSubsystem SubsystemType);
+	virtual void OnPartyInviteAcceptedInternal(const FOnlinePartyTypeId& PartyTypeId) const;
+	virtual void OnPartyInviteRejectedInternal(const FOnlinePartyTypeId& PartyTypeId) const;
 	virtual void SetSubsystemId(ESocialSubsystem SubsystemType, const FUniqueNetIdRepl& SubsystemId);
 	int32 NumPendingQueries = 0;
 
