Change order of operations for changing party leader
Previously, would trigger any OnDemoted events before the party leader ID was updated, making any code which checked Party Leader behave incorrectly.
Now, party leader ID is updated, then Demotion events are triggered if there was a previous leader, then promotion events.

diff --git a/Engine/Plugins/Online/OnlineFramework/Source/Party/Private/Party/SocialParty.cpp b/Engine/Plugins/Online/OnlineFramework/Source/Party/Private/Party/SocialParty.cpp
--- a/Engine/Plugins/Online/OnlineFramework/Source/Party/Private/Party/SocialParty.cpp
+++ b/Engine/Plugins/Online/OnlineFramework/Source/Party/Private/Party/SocialParty.cpp
@@ -774,20 +774,19 @@ void USocialParty::HandlePartyMemberPromoted(const FUniqueNetId& LocalUserId, co
 	{
 		UE_LOG(LogParty, VeryVerbose, TEXT("Party member [%s] in party [%s] promoted"), *NewLeaderId.ToDebugString(), *PartyId.ToDebugString());
 
-		if (CurrentLeaderId.IsValid() && NewLeaderId != *CurrentLeaderId)
+		UPartyMember* PreviousLeader = GetPartyMember(CurrentLeaderId);
+
+		CurrentLeaderId = NewLeaderId.AsShared();
+
+		if (PreviousLeader)
 		{
-			if (UPartyMember* PreviousLeader = GetPartyMember(CurrentLeaderId))
+			PreviousLeader->NotifyMemberDemoted();
+			if (PreviousLeader->IsLocalPlayer())
 			{
-				PreviousLeader->NotifyMemberDemoted();
-				if (PreviousLeader->IsLocalPlayer())
-				{
-					OnLocalPlayerIsLeaderChanged(false);
-				}
+				OnLocalPlayerIsLeaderChanged(false);
 			}
 		}
 
-		CurrentLeaderId = NewLeaderId.AsShared();
-
 		UPartyMember* NewLeader = GetPartyMember(CurrentLeaderId);
 		if (ensure(NewLeader))
 		{
