Resolving the possibility for a newly created party to never update its party config because the desired privacy settings didn't differ from the default.
- Thing is, the default privacy settings don't line up with the default config for new parties. Now the defaults on FPartyPrivacySettings line up with the maximum privacy defaults we use when creating new parties.

#tests PIE 2 clients repro'd the issue, worked correctly following these fixes

diff --git a/Engine/Plugins/Online/OnlineFramework/Source/Party/Public/Party/PartyTypes.h b/Engine/Plugins/Online/OnlineFramework/Source/Party/Public/Party/PartyTypes.h
--- a/Engine/Plugins/Online/OnlineFramework/Source/Party/Public/Party/PartyTypes.h
+++ b/Engine/Plugins/Online/OnlineFramework/Source/Party/Public/Party/PartyTypes.h
@@ -263,15 +263,15 @@ struct FPartyPrivacySettings
 public:
 	/** The type of party in terms of advertised joinability restrictions */
 	UPROPERTY()
-	EPartyType PartyType = EPartyType::FriendsOnly;
+	EPartyType PartyType = EPartyType::Private;
 
 	/** Who is allowed to send invitataions to the party? */
 	UPROPERTY()
-	EPartyInviteRestriction PartyInviteRestriction = EPartyInviteRestriction::AnyMember;
+	EPartyInviteRestriction PartyInviteRestriction = EPartyInviteRestriction::NoInvites;
 
 	/** True to restrict the party exclusively to friends of the party leader */
 	UPROPERTY()
-	bool bOnlyLeaderFriendsCanJoin = false;
+	bool bOnlyLeaderFriendsCanJoin = true;
 
 	bool operator==(const FPartyPrivacySettings& Other) const;
 	bool operator!=(const FPartyPrivacySettings& Other) const { return !operator==(Other); }
