[SocialToolkit] Fix issue where party invites that occured _before_ the user logs in/readfriendslist completes would not correctly display.

diff --git a/Engine/Plugins/Online/OnlineFramework/Source/Party/Private/SocialToolkit.cpp b/Engine/Plugins/Online/OnlineFramework/Source/Party/Private/SocialToolkit.cpp
--- a/Engine/Plugins/Online/OnlineFramework/Source/Party/Private/SocialToolkit.cpp
+++ b/Engine/Plugins/Online/OnlineFramework/Source/Party/Private/SocialToolkit.cpp
@@ -648,6 +648,8 @@ void USocialToolkit::HandleReadFriendsListComplete(int32 LocalUserNum, bool bWas
 		//@todo DanH: This is a really big deal on primary and a frustrating deal on platform
 		// In both cases I think we should give it another shot, but I dunno how long to wait and if we should behave differently between the two
 	}
+
+	HandleExistingPartyInvites(SubsystemType);
 }
 
 void USocialToolkit::HandleQueryBlockedPlayersComplete(const FUniqueNetId& UserId, bool bWasSuccessful, const FString& ErrorStr, ESocialSubsystem SubsystemType)
@@ -963,3 +965,24 @@ void USocialToolkit::HandleUserInvalidated(USocialUser* InvalidUser)
 {
 	AllUsers.Remove(InvalidUser);
 }
+
+void USocialToolkit::HandleExistingPartyInvites(ESocialSubsystem SubsystemType)
+{
+	if (SubsystemType == ESocialSubsystem::Primary)
+	{
+		if (IOnlineSubsystem* Oss = GetSocialOss(SubsystemType))
+		{
+			IOnlinePartyPtr PartyInterface = Oss->GetPartyInterface();
+			FUniqueNetIdRepl LocalUserId = GetLocalUserNetId(SubsystemType);
+			if (PartyInterface.IsValid() && LocalUserId.IsValid())
+			{
+				TArray<TSharedRef<IOnlinePartyJoinInfo>> PendingInvites;
+				PartyInterface->GetPendingInvites(*LocalUserId, PendingInvites);
+				for (const TSharedRef<IOnlinePartyJoinInfo>& PendingInvite : PendingInvites)
+				{
+					HandlePartyInviteReceived(*LocalUserId, *PendingInvite->GetPartyId(), *PendingInvite->GetSourceUserId());
+				}
+			}
+		}
+	}
+}
diff --git a/Engine/Plugins/Online/OnlineFramework/Source/Party/Public/SocialToolkit.h b/Engine/Plugins/Online/OnlineFramework/Source/Party/Public/SocialToolkit.h
--- a/Engine/Plugins/Online/OnlineFramework/Source/Party/Public/SocialToolkit.h
+++ b/Engine/Plugins/Online/OnlineFramework/Source/Party/Public/SocialToolkit.h
@@ -201,6 +201,8 @@ private:	// Handlers
 	void HandlePartyMemberExited(const FUniqueNetId& LocalUserId, const FOnlinePartyId& PartyId, const FUniqueNetId& MemberId, const EMemberExitedReason Reason);
 	void HandleGameDestroyed(const FName SessionName, bool bWasSuccessful);
 
+	void HandleExistingPartyInvites(ESocialSubsystem SubsystemType);
+
 private:
 	static USocialToolkit* GetToolkitForPlayerInternal(ULocalPlayer* LocalPlayer);
 	static TMap<TWeakObjectPtr<ULocalPlayer>, TWeakObjectPtr<USocialToolkit>> AllToolkitsByOwningPlayer;
