Add config settings for whether we should leave a party on disconnect, and whether we should cache parties for join on disconnect

diff --git a/Engine/Plugins/Online/OnlineFramework/Source/Party/Private/Party/SocialParty.cpp b/Engine/Plugins/Online/OnlineFramework/Source/Party/Private/Party/SocialParty.cpp
--- a/Engine/Plugins/Online/OnlineFramework/Source/Party/Private/Party/SocialParty.cpp
+++ b/Engine/Plugins/Online/OnlineFramework/Source/Party/Private/Party/SocialParty.cpp
@@ -95,14 +95,6 @@ const USocialParty* FPartyRepData::GetOwnerParty() const
 // USocialParty
 //////////////////////////////////////////////////////////////////////////
 
-static int32 EnableAutomaticPartyRejoin = 1;
-static FAutoConsoleVariableRef CVarEnableAutomaticPartyRejoin(
-	TEXT("Party.EnableAutomaticPartyRejoin"),
-	EnableAutomaticPartyRejoin,
-	TEXT("Enable automatic rejoining of parties\n")
-	TEXT("1 Enables. 0 disables."),
-	ECVF_Default);
-
 static int32 AllowPartyJoinsDuringLoad = 1;
 static FAutoConsoleVariableRef CVar_AllowPartyJoinsDuringLoad(
 	TEXT("Party.AllowJoinsDuringLoad"),
@@ -226,7 +218,7 @@ FPartyJoinApproval USocialParty::EvaluateJoinRequest(const FUniqueNetId& PlayerI
 
 bool USocialParty::ShouldCacheForRejoinOnDisconnect() const
 {
-	return EnableAutomaticPartyRejoin != 0 && GetNumPartyMembers() > 1;
+	return bEnableAutomaticPartyRejoin && GetNumPartyMembers() > 1;
 }
 
 bool USocialParty::IsCurrentlyLeaving() const
diff --git a/Engine/Plugins/Online/OnlineFramework/Source/Party/Private/SocialManager.cpp b/Engine/Plugins/Online/OnlineFramework/Source/Party/Private/SocialManager.cpp
--- a/Engine/Plugins/Online/OnlineFramework/Source/Party/Private/SocialManager.cpp
+++ b/Engine/Plugins/Online/OnlineFramework/Source/Party/Private/SocialManager.cpp
@@ -991,7 +991,10 @@ void USocialManager::HandlePersistentPartyStateChanged(EPartyState NewState, EPa
 		
 		if (PreviousState == EPartyState::Active)
 		{
-			PersistentParty->LeaveParty();
+			if (bLeavePartyOnDisconnect)
+			{
+				PersistentParty->LeaveParty();
+			}
 		}
 
 		// If we have other members in our party, then we will try to rejoin this when we come back online
diff --git a/Engine/Plugins/Online/OnlineFramework/Source/Party/Public/Party/SocialParty.h b/Engine/Plugins/Online/OnlineFramework/Source/Party/Public/Party/SocialParty.h
--- a/Engine/Plugins/Online/OnlineFramework/Source/Party/Public/Party/SocialParty.h
+++ b/Engine/Plugins/Online/OnlineFramework/Source/Party/Public/Party/SocialParty.h
@@ -337,6 +337,9 @@ private:
 	UPROPERTY()
 	TMap<FUniqueNetIdRepl, UPartyMember*> PartyMembersById;
 
+	UPROPERTY(config)
+	bool bEnableAutomaticPartyRejoin = true;
+
 	FPartyConfiguration CurrentConfig;
 
 	//@todo DanH Party: Rename/reorg this to more clearly call out that this is specific to lobby beacon stuff #suggested
diff --git a/Engine/Plugins/Online/OnlineFramework/Source/Party/Public/SocialManager.h b/Engine/Plugins/Online/OnlineFramework/Source/Party/Public/SocialManager.h
--- a/Engine/Plugins/Online/OnlineFramework/Source/Party/Public/SocialManager.h
+++ b/Engine/Plugins/Online/OnlineFramework/Source/Party/Public/SocialManager.h
@@ -242,6 +242,10 @@ private:
 	 */
 	bool bCanCreatePartyObjects = false;
 
+	/** Should we leave a party when it enters the disconnected state? */
+	UPROPERTY(config)
+	bool bLeavePartyOnDisconnect = true;
+
 	TSharedPtr<FPartyPlatformSessionManager> PartySessionManager;
 
 	TMap<FOnlinePartyTypeId, USocialParty*> JoinedPartiesByTypeId;
