Cleanup - Updating several FTimerDelegate::CreateLambda calls to be, FTimerDelegate::CreateWeakLambda(this, [this,..., so that we can benefit from a lambda that's bound to the lifetime of the object without having to manually create a WeakThis ptr on the stack and capture into the lambda, or in some cases passing a raw this ptr, into the capture set, and hope the UObject isn't GCed before the Lambda is called.

Michael.Noland, Ian.Fox, Peter.Knepley, Dan.Hertzka

diff --git a/Engine/Plugins/Online/OnlineFramework/Source/Party/Private/Chatroom.cpp b/Engine/Plugins/Online/OnlineFramework/Source/Party/Private/Chatroom.cpp
--- a/Engine/Plugins/Online/OnlineFramework/Source/Party/Private/Chatroom.cpp
+++ b/Engine/Plugins/Online/OnlineFramework/Source/Party/Private/Chatroom.cpp
@@ -142,7 +142,7 @@ void UChatroom::OnChatRoomCreatedOrJoined(const FUniqueNetId& LocalUserId, const
 
 					FUniqueNetIdRepl StrongUserId(LocalUserId.AsShared());
 
-					auto RetryDelegate = FTimerDelegate::CreateLambda([this, StrongUserId, RoomId, CompletionDelegate, RoomConfig] ()
+					auto RetryDelegate = FTimerDelegate::CreateWeakLambda(this, [this, StrongUserId, RoomId, CompletionDelegate, RoomConfig] ()
 					{
 						// Attempt to rejoin room shortly
 						CreateOrJoinChatRoom(StrongUserId, RoomId, CompletionDelegate, RoomConfig);
diff --git a/Engine/Plugins/Online/OnlineFramework/Source/Party/Private/Party/SocialParty.cpp b/Engine/Plugins/Online/OnlineFramework/Source/Party/Private/Party/SocialParty.cpp
--- a/Engine/Plugins/Online/OnlineFramework/Source/Party/Private/Party/SocialParty.cpp
+++ b/Engine/Plugins/Online/OnlineFramework/Source/Party/Private/Party/SocialParty.cpp
@@ -959,10 +959,10 @@ void USocialParty::HandlePartySystemStateChange(EPartySystemState NewState)
 		SetIsRequestingShutdown(true);
 
 		//set timer to turn this off in a minute?
-		UWorld* World = GetWorld();
 		FTimerHandle DummyHandle;
-		TWeakObjectPtr<USocialParty> WeakThis(this);
-		World->GetTimerManager().SetTimer(DummyHandle, FTimerDelegate::CreateLambda([WeakThis, this]() { if (WeakThis.IsValid()) { SetIsRequestingShutdown(false); } }), 60.0f, false);
+		GetWorld()->GetTimerManager().SetTimer(DummyHandle, FTimerDelegate::CreateWeakLambda(this, [this]() {
+			SetIsRequestingShutdown(false);
+		}), 60.0f, false);
 	}
 }
 
