Adding SuggestedFriend social relationship to the social user list.
  -- Also SameParty presence filter shouldn't trigger the "friends only" presence warning.

diff --git a/Engine/Plugins/Online/OnlineFramework/Source/Party/Private/User/SocialUser.cpp b/Engine/Plugins/Online/OnlineFramework/Source/Party/Private/User/SocialUser.cpp
--- a/Engine/Plugins/Online/OnlineFramework/Source/Party/Private/User/SocialUser.cpp
+++ b/Engine/Plugins/Online/OnlineFramework/Source/Party/Private/User/SocialUser.cpp
@@ -246,6 +246,12 @@ TArray<ESocialSubsystem> USocialUser::GetRelationshipSubsystems(ESocialRelations
 					RelationshipSubsystems.Add(SubsystemInfoPair.Key);
 				}
 				break;
+			case ESocialRelationship::SuggestedFriend:
+				if (SubsystemInfoPair.Value.GetFriendInviteStatus() == EInviteStatus::Suggested)
+				{
+					RelationshipSubsystems.Add(SubsystemInfoPair.Key);
+				}
+				break;
 			case ESocialRelationship::Friend:
 				if (SubsystemInfoPair.Value.IsFriend())
 				{
diff --git a/Engine/Plugins/Online/OnlineFramework/Source/Party/Private/User/SocialUserList.cpp b/Engine/Plugins/Online/OnlineFramework/Source/Party/Private/User/SocialUserList.cpp
--- a/Engine/Plugins/Online/OnlineFramework/Source/Party/Private/User/SocialUserList.cpp
+++ b/Engine/Plugins/Online/OnlineFramework/Source/Party/Private/User/SocialUserList.cpp
@@ -23,12 +23,12 @@ FSocialUserList::FSocialUserList(USocialToolkit& InOwnerToolkit, const FSocialUs
 	: OwnerToolkit(&InOwnerToolkit)
 	, ListConfig(InConfig)
 {
-	if (HasPresenceFilters() &&
+	if (HasPresenceFilters() && ListConfig.RequiredPresenceFlags != ESocialUserStateFlags::SameParty && ListConfig.ForbiddenPresenceFlags != ESocialUserStateFlags::SameParty &&
 		ListConfig.RelationshipType != ESocialRelationship::Any &&
 		ListConfig.RelationshipType != ESocialRelationship::Friend &&
 		ListConfig.RelationshipType != ESocialRelationship::PartyInvite)
 	{
-		UE_LOG(LogParty, Error, TEXT("A user list with presence filters can only ever track friends. No users will ever appear in this list."));
+		UE_LOG(LogParty, Error, TEXT("A user list with friend presence filters can only ever track friends. No users will ever appear in this list."));
 	}
 }
 
@@ -54,6 +54,9 @@ void FSocialUserList::InitializeList()
 		OwnerToolkit->OnRecentPlayerAdded().AddSP(this, &FSocialUserList::HandleRecentPlayerAdded);
 		OwnerToolkit->OnFriendshipEstablished().AddSP(this, &FSocialUserList::HandleFriendshipEstablished);
 		break;
+	case ESocialRelationship::SuggestedFriend:
+		OwnerToolkit->OnFriendshipEstablished().AddSP(this, &FSocialUserList::HandleFriendshipEstablished);
+		break;
 	}
 
 	OwnerToolkit->OnToolkitReset().AddSP(this, &FSocialUserList::HandleOwnerToolkitReset);
@@ -165,8 +168,7 @@ void FSocialUserList::HandleFriendshipEstablished(USocialUser& NewFriend, ESocia
 		TryAddUser(NewFriend);
 	}
 
-	if (ListConfig.RelationshipType != ESocialRelationship::Friend ||
-		ListConfig.RelationshipType == ESocialRelationship::FriendInviteReceived)
+	if (ListConfig.RelationshipType != ESocialRelationship::Friend)
 	{
 		// Any non-friends list that cares about friendship does so to remove entries (i.e. invites & recent players)
 		TryRemoveUser(NewFriend);
@@ -371,8 +373,6 @@ void FSocialUserList::TryRemoveUserFast(USocialUser& User)
 		User.OnPartyInviteAccepted().RemoveAll(this);
 		User.OnPartyInviteRejected().RemoveAll(this);
 		User.OnBlockedStatusChanged().RemoveAll(this);
-		User.OnPartyInviteAccepted().RemoveAll(this);
-		User.OnPartyInviteRejected().RemoveAll(this);
 
 		if (bUnbindFromPresenceUpdates)
 		{
@@ -476,12 +476,6 @@ bool FSocialUserList::HandleAutoUpdateList(float)
 
 void FSocialUserList::UpdateListInternal()
 {
-	if (NumIgnoreUpdateRequests > 0)
-	{
-		return;
-	}
-
-	
 	// Re-evaluate whether each user with dirtied presence is still fit for the list
 	for (TWeakObjectPtr<USocialUser> DirtyUser : UsersWithDirtyPresence)
 	{
@@ -563,13 +557,11 @@ void FSocialUserList::HandlePartyJoined(USocialParty& Party)
 {
 	Party.OnPartyMemberCreated().AddSP(this, &FSocialUserList::HandlePartyMemberCreated);
 
-	// ignore updates while we populate the party members (lots of redundant actions)
-	NumIgnoreUpdateRequests++;
 	for (UPartyMember* PartyMember : Party.GetPartyMembers())
 	{
-		HandlePartyMemberCreated(*PartyMember);
+		PartyMember->OnLeftParty().AddSP(this, &FSocialUserList::HandlePartyMemberLeft, PartyMember);
+		MarkUserAsDirty(PartyMember->GetSocialUser());
 	}
-	NumIgnoreUpdateRequests--;
 
 	UpdateNow();
 }
diff --git a/Engine/Plugins/Online/OnlineFramework/Source/Party/Private/User/SocialUserList.h b/Engine/Plugins/Online/OnlineFramework/Source/Party/Private/User/SocialUserList.h
--- a/Engine/Plugins/Online/OnlineFramework/Source/Party/Private/User/SocialUserList.h
+++ b/Engine/Plugins/Online/OnlineFramework/Source/Party/Private/User/SocialUserList.h
@@ -82,7 +82,6 @@ private:
 	FSocialUserListConfig ListConfig;
 
 	// give external access to disable list update for perf
-	int32 NumIgnoreUpdateRequests = 0;
 	bool bAllowAutoUpdate = true;
 	bool bNeedsSort = false;
 	float AutoUpdatePeriod = 5.f;
