[VOIP] Replicate VoiceConnectionId in FPartyMemberRepData

diff --git a/Engine/Plugins/Online/OnlineFramework/Source/Party/Private/Party/PartyMember.cpp b/Engine/Plugins/Online/OnlineFramework/Source/Party/Private/Party/PartyMember.cpp
--- a/Engine/Plugins/Online/OnlineFramework/Source/Party/Private/Party/PartyMember.cpp
+++ b/Engine/Plugins/Online/OnlineFramework/Source/Party/Private/Party/PartyMember.cpp
@@ -31,6 +31,7 @@ void FPartyMemberRepData::CompareAgainst(const FOnlinePartyRepDataBase& OldData)
 	ComparePlatformUniqueId(TypedOldData);
 	ComparePlatformSessionId(TypedOldData);
 	CompareCrossplayPreference(TypedOldData);
+	CompareVoiceConnectionId(TypedOldData);
 }
 
 const USocialParty* FPartyMemberRepData::GetOwnerParty() const
diff --git a/Engine/Plugins/Online/OnlineFramework/Source/Party/Public/Party/PartyMember.h b/Engine/Plugins/Online/OnlineFramework/Source/Party/Public/Party/PartyMember.h
--- a/Engine/Plugins/Online/OnlineFramework/Source/Party/Public/Party/PartyMember.h
+++ b/Engine/Plugins/Online/OnlineFramework/Source/Party/Public/Party/PartyMember.h
@@ -54,6 +54,11 @@ private:
 	UPROPERTY()
 	ECrossplayPreference CrossplayPreference = ECrossplayPreference::NoSelection;
 	EXPOSE_REP_DATA_PROPERTY(FPartyMemberRepData, ECrossplayPreference, CrossplayPreference);
+
+	/** The voice device selected by the user. Can be blank if the user has not pinned a device */
+	UPROPERTY()
+	FString VoiceConnectionId;
+	EXPOSE_REP_DATA_PROPERTY(FPartyMemberRepData, FString, VoiceConnectionId);
 };
 
 using FPartyMemberDataReplicator = TPartyDataReplicator<FPartyMemberRepData>;
