Added support for allowing game-level subclasses of SocialUser to provide last-level sorting information for UserLists.

Added functionality to SocialUser and SocialParty that allows a user to be sent an invite when IsOnline is false.

#tests Compile DebugGame Editor Win64 / Verified proper sorting with game override in PIE, and sending of offline invites

diff --git a/Engine/Plugins/Online/OnlineFramework/Source/Party/Private/Party/SocialParty.cpp b/Engine/Plugins/Online/OnlineFramework/Source/Party/Private/Party/SocialParty.cpp
--- a/Engine/Plugins/Online/OnlineFramework/Source/Party/Private/Party/SocialParty.cpp
+++ b/Engine/Plugins/Online/OnlineFramework/Source/Party/Private/Party/SocialParty.cpp
@@ -262,7 +262,7 @@ bool USocialParty::CanInviteUser(const USocialUser& User) const
 bool USocialParty::CanInviteUserInternal(const USocialUser& User) const
 {
 	// Only users that are online can be invited
-	if (!User.IsOnline())
+	if (!User.IsOnline() && !User.CanReceiveOfflineInvite())
 	{
 		return false;
 	}
diff --git a/Engine/Plugins/Online/OnlineFramework/Source/Party/Private/User/SocialUserList.cpp b/Engine/Plugins/Online/OnlineFramework/Source/Party/Private/User/SocialUserList.cpp
--- a/Engine/Plugins/Online/OnlineFramework/Source/Party/Private/User/SocialUserList.cpp
+++ b/Engine/Plugins/Online/OnlineFramework/Source/Party/Private/User/SocialUserList.cpp
@@ -434,14 +434,15 @@ bool FSocialUserList::EvaluatePresenceFlag(bool bPresenceValue, ESocialUserState
 // encapsulates UserList sorting comparator and supporting data needed
 struct FUserSortData
 {
-	FUserSortData(USocialUser* InUser, EOnlinePresenceState::Type InStatus, bool InPlayingThisGame, FString InDisplayName)
-		: User(InUser), OnlineStatus(InStatus), PlayingThisGame(InPlayingThisGame), DisplayName(MoveTemp(InDisplayName))
+	FUserSortData(USocialUser* InUser, EOnlinePresenceState::Type InStatus, bool InPlayingThisGame, FString InDisplayName, int32 InCustomSortValue)
+		: User(InUser), OnlineStatus(InStatus), PlayingThisGame(InPlayingThisGame), DisplayName(MoveTemp(InDisplayName)), CustomSortValue(InCustomSortValue)
 	{ }
 
 	USocialUser* User;
 	EOnlinePresenceState::Type OnlineStatus;
 	bool PlayingThisGame;
 	FString DisplayName;
+	int32 CustomSortValue;
 
 	bool operator<(const FUserSortData& OtherSortData) const
 	{
@@ -450,7 +451,14 @@ struct FUserSortData
 		{
 			if (PlayingThisGame == OtherSortData.PlayingThisGame)
 			{
-				return DisplayName < OtherSortData.DisplayName;
+				if (CustomSortValue == OtherSortData.CustomSortValue)
+				{
+					return DisplayName < OtherSortData.DisplayName;
+				}
+				else
+				{
+					return CustomSortValue > OtherSortData.CustomSortValue;
+				}
 			}
 			else
 			{
@@ -538,7 +546,7 @@ void FSocialUserList::UpdateListInternal()
 
 		Algo::Transform(Users, SortedData, [](USocialUser* const User) -> FUserSortData
 		{
-			return FUserSortData(User, User->GetOnlineStatus(), User->IsPlayingThisGame(), User->GetDisplayName());
+			return FUserSortData(User, User->GetOnlineStatus(), User->IsPlayingThisGame(), User->GetDisplayName(), User->GetCustomSortValue());
 		});
 
 		Algo::Sort(SortedData);
diff --git a/Engine/Plugins/Online/OnlineFramework/Source/Party/Public/User/SocialUser.h b/Engine/Plugins/Online/OnlineFramework/Source/Party/Public/User/SocialUser.h
--- a/Engine/Plugins/Online/OnlineFramework/Source/Party/Public/User/SocialUser.h
+++ b/Engine/Plugins/Online/OnlineFramework/Source/Party/Public/User/SocialUser.h
@@ -64,6 +64,9 @@ public:
 	bool IsOnline() const;
 	bool IsPlayingThisGame() const;
 	
+	virtual bool CanReceiveOfflineInvite() const { return false; }
+	virtual int32 GetCustomSortValue() const { return 0; }
+
 	bool SetUserLocalAttribute(ESocialSubsystem SubsystemType, const FString& AttrName, const FString& AttrValue);
 	bool GetUserAttribute(ESocialSubsystem SubsystemType, const FString& AttrName, FString& OutAttrValue) const;
 
