Move VoiceConnectionId out of OnlineFramework

diff --git a/Engine/Plugins/Online/OnlineFramework/Source/Party/Private/Party/PartyMember.cpp b/Engine/Plugins/Online/OnlineFramework/Source/Party/Private/Party/PartyMember.cpp
--- a/Engine/Plugins/Online/OnlineFramework/Source/Party/Private/Party/PartyMember.cpp
+++ b/Engine/Plugins/Online/OnlineFramework/Source/Party/Private/Party/PartyMember.cpp
@@ -31,7 +31,6 @@ void FPartyMemberRepData::CompareAgainst(const FOnlinePartyRepDataBase& OldData)
 	ComparePlatformUniqueId(TypedOldData);
 	ComparePlatformSessionId(TypedOldData);
 	CompareCrossplayPreference(TypedOldData);
-	CompareVoiceConnectionId(TypedOldData);
 }
 
 const USocialParty* FPartyMemberRepData::GetOwnerParty() const
diff --git a/Engine/Plugins/Online/OnlineFramework/Source/Party/Public/Party/PartyMember.h b/Engine/Plugins/Online/OnlineFramework/Source/Party/Public/Party/PartyMember.h
--- a/Engine/Plugins/Online/OnlineFramework/Source/Party/Public/Party/PartyMember.h
+++ b/Engine/Plugins/Online/OnlineFramework/Source/Party/Public/Party/PartyMember.h
@@ -54,11 +54,6 @@ private:
 	UPROPERTY()
 	ECrossplayPreference CrossplayPreference = ECrossplayPreference::NoSelection;
 	EXPOSE_REP_DATA_PROPERTY(FPartyMemberRepData, ECrossplayPreference, CrossplayPreference);
-
-	/** The voice device selected by the user. Can be blank if the user has not pinned a device */
-	UPROPERTY()
-	FString VoiceConnectionId;
-	EXPOSE_REP_DATA_PROPERTY(FPartyMemberRepData, FString, VoiceConnectionId);
 };
 
 using FPartyMemberDataReplicator = TPartyDataReplicator<FPartyMemberRepData>;
