Modified the NotifyInviteFailed function in SocialToolkit to be a general OnInviteComplete function and made it pass invite errors as strings rather than enums.

#tests Compile DebugGame Editor Win64

diff --git a/Engine/Plugins/Online/OnlineFramework/Source/Party/Private/SocialToolkit.cpp b/Engine/Plugins/Online/OnlineFramework/Source/Party/Private/SocialToolkit.cpp
--- a/Engine/Plugins/Online/OnlineFramework/Source/Party/Private/SocialToolkit.cpp
+++ b/Engine/Plugins/Online/OnlineFramework/Source/Party/Private/SocialToolkit.cpp
@@ -765,11 +765,11 @@ void USocialToolkit::HandleQueryPrimaryUserIdMappingComplete(bool bWasSuccessful
 {
 	if (!IdentifiedUserId.IsValid())
 	{
-		NotifyFriendInviteFailed(IdentifiedUserId, DisplayName, ESendFriendInviteFailureReason::NotFound);
+		OnFriendInviteComplete(IdentifiedUserId, DisplayName, false, FriendInviteFailureReason::InviteFailReason_NotFound);
 	}
 	else if (RequestingUserId == IdentifiedUserId)
 	{
-		NotifyFriendInviteFailed(IdentifiedUserId, DisplayName, ESendFriendInviteFailureReason::AddingSelfFail);
+		OnFriendInviteComplete(IdentifiedUserId, DisplayName, false, FriendInviteFailureReason::InviteFailReason_AddingSelfFail);
 	}
 	else
 	{
@@ -778,11 +778,11 @@ void USocialToolkit::HandleQueryPrimaryUserIdMappingComplete(bool bWasSuccessful
 			{
 				if (SocialUser.IsBlocked())
 				{
-					NotifyFriendInviteFailed(*SocialUser.GetUserId(ESocialSubsystem::Primary), DisplayName, ESendFriendInviteFailureReason::AddingBlockedFail);
+					OnFriendInviteComplete(*SocialUser.GetUserId(ESocialSubsystem::Primary), DisplayName, false, FriendInviteFailureReason::InviteFailReason_AddingBlockedFail);
 				}
 				else if (SocialUser.IsFriend(ESocialSubsystem::Primary))
 				{
-					NotifyFriendInviteFailed(*SocialUser.GetUserId(ESocialSubsystem::Primary), DisplayName, ESendFriendInviteFailureReason::AlreadyFriends);
+					OnFriendInviteComplete(*SocialUser.GetUserId(ESocialSubsystem::Primary), DisplayName, false, FriendInviteFailureReason::InviteFailReason_AlreadyFriends);
 				}
 				else
 				{
@@ -863,10 +863,7 @@ void USocialToolkit::HandleFriendInviteSent(int32 LocalUserNum, bool bWasSuccess
 				}
 			});
 	}
-	else
-	{
-		NotifyFriendInviteFailed(InvitedUserId, ErrorStr, ESendFriendInviteFailureReason::UnknownError, false);
-	}
+	OnFriendInviteComplete(InvitedUserId, DisplayName, bWasSuccessful, ErrorStr);
 }
 
 void USocialToolkit::HandleFriendRemoved(const FUniqueNetId& LocalUserId, const FUniqueNetId& FriendId, ESocialSubsystem SubsystemType)
diff --git a/Engine/Plugins/Online/OnlineFramework/Source/Party/Public/SocialToolkit.h b/Engine/Plugins/Online/OnlineFramework/Source/Party/Public/SocialToolkit.h
--- a/Engine/Plugins/Online/OnlineFramework/Source/Party/Public/SocialToolkit.h
+++ b/Engine/Plugins/Online/OnlineFramework/Source/Party/Public/SocialToolkit.h
@@ -19,6 +19,14 @@ class FOnlinePartyId;
 
 enum class EMemberExitedReason : uint8;
 
+namespace FriendInviteFailureReason 
+{
+	const FString InviteFailReason_NotFound = TEXT("NotFound");
+	const FString InviteFailReason_AddingSelfFail = TEXT("AddingSelfFail");
+	const FString InviteFailReason_AddingBlockedFail = TEXT("AddingBlockedFail");
+	const FString InviteFailReason_AlreadyFriends = TEXT("AlreadyFriends");
+}
+
 namespace EOnlinePresenceState { enum Type : uint8; }
 
 DECLARE_DELEGATE_OneParam(FUserDependentAction, USocialUser&);
@@ -120,7 +128,7 @@ protected:
 	virtual void OnOwnerLoggedIn();
 	virtual void OnOwnerLoggedOut();
 	
-	virtual void NotifyFriendInviteFailed(const FUniqueNetId& InvitedUserId, const FString& InvitedUserName, ESendFriendInviteFailureReason FailureReason, bool bCanShow = true) {}
+	virtual void OnFriendInviteComplete(const FUniqueNetId& InvitedUserId, const FString& InvitedUserName, bool bInviteSuceeded, const FString& InviteError) {}
 
 	void QueryFriendsLists();
 	void QueryBlockedPlayers();
diff --git a/Engine/Plugins/Online/OnlineFramework/Source/Party/Public/SocialTypes.h b/Engine/Plugins/Online/OnlineFramework/Source/Party/Public/SocialTypes.h
--- a/Engine/Plugins/Online/OnlineFramework/Source/Party/Public/SocialTypes.h
+++ b/Engine/Plugins/Online/OnlineFramework/Source/Party/Public/SocialTypes.h
@@ -92,17 +92,6 @@ enum class ECrossplayPreference : uint8
 	OptedOutRestricted
 };
 
-UENUM()
-enum class ESendFriendInviteFailureReason : uint8
-{
-	NotFound,
-	AlreadyFriends,
-	InvitePending,
-	AddingSelfFail,
-	AddingBlockedFail,
-	UnknownError
-};
-
 /** Thin wrapper to infuse a raw platform string with some meaning */
 USTRUCT()
 struct PARTY_API FUserPlatform
