Attempt to move leadership of party to another player when disconnecting from party.

diff --git a/Engine/Plugins/Online/OnlineFramework/Source/Party/Public/Party/SocialParty.h b/Engine/Plugins/Online/OnlineFramework/Source/Party/Public/Party/SocialParty.h
--- a/Engine/Plugins/Online/OnlineFramework/Source/Party/Public/Party/SocialParty.h
+++ b/Engine/Plugins/Online/OnlineFramework/Source/Party/Public/Party/SocialParty.h
@@ -146,7 +146,7 @@ public:
 	/** Is the specified net driver for our reservation beacon? */
 	bool IsNetDriverFromReservationBeacon(const UNetDriver* InNetDriver) const;
 
-	void DisconnectParty();
+	virtual void DisconnectParty();
 
 	template <typename MemberT = UPartyMember>
 	TArray<MemberT*> GetPartyMembers() const
