Block interaction shouldn't be available on Local User

diff --git a/Engine/Plugins/Online/OnlineFramework/Source/Party/Private/Interactions/CoreInteractions.cpp b/Engine/Plugins/Online/OnlineFramework/Source/Party/Private/Interactions/CoreInteractions.cpp
--- a/Engine/Plugins/Online/OnlineFramework/Source/Party/Private/Interactions/CoreInteractions.cpp
+++ b/Engine/Plugins/Online/OnlineFramework/Source/Party/Private/Interactions/CoreInteractions.cpp
@@ -168,7 +168,7 @@ FString FSocialInteraction_Block::GetSlashCommandToken()
 
 bool FSocialInteraction_Block::CanExecute(const USocialUser& User)
 {
-	return User.HasSubsystemInfo(ESocialSubsystem::Primary) && !User.IsBlocked(ESocialSubsystem::Primary);
+	return User.HasSubsystemInfo(ESocialSubsystem::Primary) && !User.IsBlocked(ESocialSubsystem::Primary) && !User.IsLocalUser();
 }
 
 void FSocialInteraction_Block::ExecuteInteraction(USocialUser& User)
