Engine Social changes to support additional filtering of user-provided strings with rich text markup. Social User now has virtual SanitizePresenceString that game clients can hook into to provide additional filtering of presence text.

Part 1 of 2 submits.

#tests PIE no lets chat messages use user-provided markup. Pressence status runs through an additional filter to remove markup. Compile win64, ps4 shipping client.
nick.darnell

diff --git a/Engine/Plugins/Online/OnlineFramework/Source/Party/Private/User/SocialUser.cpp b/Engine/Plugins/Online/OnlineFramework/Source/Party/Private/User/SocialUser.cpp
--- a/Engine/Plugins/Online/OnlineFramework/Source/Party/Private/User/SocialUser.cpp
+++ b/Engine/Plugins/Online/OnlineFramework/Source/Party/Private/User/SocialUser.cpp
@@ -560,14 +560,14 @@ void USocialUser::GetRichPresenceText(FText& OutRichPresence) const
 		const FOnlineUserPresence* PrimaryPresence = GetFriendPresenceInfo(ESocialSubsystem::Primary);
 		if (PrimaryPresence && !PrimaryPresence->Status.StatusStr.IsEmpty())
 		{
-			OutRichPresence = FText::FromString(PrimaryPresence->Status.StatusStr);
+			OutRichPresence = FText::FromString(SanitizePresenceString(PrimaryPresence->Status.StatusStr));
 		}
 		else
 		{
 			const FOnlineUserPresence* PlatformPresence = GetFriendPresenceInfo(ESocialSubsystem::Platform);
 			if (PlatformPresence && !PlatformPresence->Status.StatusStr.IsEmpty())
 			{
-				OutRichPresence = FText::FromString(PlatformPresence->Status.StatusStr);
+				OutRichPresence = FText::FromString(SanitizePresenceString(PlatformPresence->Status.StatusStr)); 
 			}
 			else
 			{
@@ -1166,4 +1166,9 @@ void USocialUser::HandleQueryUserInfoComplete(ESocialSubsystem SubsystemType, bo
 	TryBroadcastInitializationComplete();
 }
 
+FString USocialUser::SanitizePresenceString(FString InString) const
+{
+	return InString;
+}
+
 #undef LOCTEXT_NAMESPACE
diff --git a/Engine/Plugins/Online/OnlineFramework/Source/Party/Public/User/SocialUser.h b/Engine/Plugins/Online/OnlineFramework/Source/Party/Public/User/SocialUser.h
--- a/Engine/Plugins/Online/OnlineFramework/Source/Party/Public/User/SocialUser.h
+++ b/Engine/Plugins/Online/OnlineFramework/Source/Party/Public/User/SocialUser.h
@@ -148,6 +148,8 @@ private:
 	
 	void SetUserInfo(ESocialSubsystem SubsystemType, const TSharedRef<FOnlineUser>& UserInfo);
 	void HandleQueryUserInfoComplete(ESocialSubsystem SubsystemType, bool bWasSuccessful, const TSharedPtr<FOnlineUser>& UserInfo);
+
+	virtual FString SanitizePresenceString(FString InString) const;
 	
 private:
 
