For players that are in launcher, party join result denail reason should return "NotPlayingThisGame" and not defaulting back to private party

diff --git a/Engine/Plugins/Online/OnlineFramework/Source/Party/Private/SocialManager.cpp b/Engine/Plugins/Online/OnlineFramework/Source/Party/Private/SocialManager.cpp
--- a/Engine/Plugins/Online/OnlineFramework/Source/Party/Private/SocialManager.cpp
+++ b/Engine/Plugins/Online/OnlineFramework/Source/Party/Private/SocialManager.cpp
@@ -496,6 +496,10 @@ FJoinPartyResult USocialManager::ValidateJoinTarget(const USocialUser& UserToJoi
 				return FPartyJoinDenialReason(EPartyJoinDenialReason::TargetUserMissingPresence);
 			}
 		}
+		else if (!UserToJoin.IsPlayingThisGame())
+		{
+			return FPartyJoinDenialReason(EPartyJoinDenialReason::TargetUserPlayingDifferentGame);
+		}
 		else
 		{
 			// We've got no info on this party for the given user, so it's gotta be private (or doesn't even exist)
