Check user and partyid before calling LeaveParty

diff --git a/Engine/Plugins/Online/OnlineFramework/Source/Party/Private/Party/SocialParty.cpp b/Engine/Plugins/Online/OnlineFramework/Source/Party/Private/Party/SocialParty.cpp
--- a/Engine/Plugins/Online/OnlineFramework/Source/Party/Private/Party/SocialParty.cpp
+++ b/Engine/Plugins/Online/OnlineFramework/Source/Party/Private/Party/SocialParty.cpp
@@ -1172,9 +1172,17 @@ void USocialParty::LeaveParty(const FOnLeavePartyAttemptComplete& OnLeaveAttempt
 
 		BeginLeavingParty(EMemberExitedReason::Left);
 
-		const IOnlinePartyPtr PartyInterface = Online::GetPartyInterfaceChecked(GetWorld());
-		FOnLeavePartyComplete OnLeaveComplete = FOnLeavePartyComplete::CreateUObject(this, &USocialParty::HandleLeavePartyComplete, OnLeaveAttemptComplete);
-		PartyInterface->LeaveParty(*OwningLocalUserId, GetPartyId(), OnLeaveComplete);
+		const FOnlinePartyId& PartyId = GetPartyId();
+		if (OwningLocalUserId.IsValid() && PartyId.IsValid())
+		{
+			const IOnlinePartyPtr PartyInterface = Online::GetPartyInterfaceChecked(GetWorld());
+			FOnLeavePartyComplete OnLeaveComplete = FOnLeavePartyComplete::CreateUObject(this, &USocialParty::HandleLeavePartyComplete, OnLeaveAttemptComplete);
+			PartyInterface->LeaveParty(*OwningLocalUserId, PartyId, OnLeaveComplete);
+		}
+		else
+		{
+			OnLeaveAttemptComplete.ExecuteIfBound(ELeavePartyCompletionResult::UnknownClientFailure);
+		}
 	}
 }
 
