Fix for game-specifc radial picker opening hamburger menu. Current flow for opening chat through FocusChannel opens hamburger menu. There exists a flow through PC for in-game chat. Both versions of focusing chat actually give chat focus (ie put a cursor in the editable text). This isn't desired for the radial picker, because players intend on using that instead of typing messages to chat. So, I created a method that only requests chat display a channel without focusing it.

#tests PIE, console in zone, shipping client compile

diff --git a/Engine/Plugins/Online/OnlineFramework/Source/Party/Public/Chat/SocialChatManager.cpp b/Engine/Plugins/Online/OnlineFramework/Source/Party/Public/Chat/SocialChatManager.cpp
--- a/Engine/Plugins/Online/OnlineFramework/Source/Party/Public/Chat/SocialChatManager.cpp
+++ b/Engine/Plugins/Online/OnlineFramework/Source/Party/Public/Chat/SocialChatManager.cpp
@@ -160,6 +160,11 @@ void USocialChatManager::FocusChatChannel(USocialChatChannel& InChannel)
 	OnChannelFocusRequestedEvent.Broadcast(InChannel);
 }
 
+void USocialChatManager::DisplayChatChannel(USocialChatChannel& InChannel)
+{
+	OnChannelDisplayRequestedEvent.Broadcast(InChannel);
+}
+
 TSubclassOf<USocialChatRoom> USocialChatManager::GetClassForChatRoom(ESocialChannelType Type) const
 {
 	if (Type == ESocialChannelType::Party)
diff --git a/Engine/Plugins/Online/OnlineFramework/Source/Party/Public/Chat/SocialChatManager.h b/Engine/Plugins/Online/OnlineFramework/Source/Party/Public/Chat/SocialChatManager.h
--- a/Engine/Plugins/Online/OnlineFramework/Source/Party/Public/Chat/SocialChatManager.h
+++ b/Engine/Plugins/Online/OnlineFramework/Source/Party/Public/Chat/SocialChatManager.h
@@ -86,10 +86,13 @@ public:
 	virtual USocialChatChannel* CreateChatChannel(FSocialChatChannelConfig& InConfig);
 
 	DECLARE_EVENT_OneParam(USocialChatManager, FOnChatChannelFocusRequested, USocialChatChannel&);
+	DECLARE_EVENT_OneParam(USocialChatManager, FOnChatChannelDisplayRequested, USocialChatChannel&);
 	FOnChatChannelFocusRequested& OnChannelFocusRequested() const { return OnChannelFocusRequestedEvent; }
+	FOnChatChannelDisplayRequested& OnChannelDisplayRequested() const { return OnChannelDisplayRequestedEvent; }
 
 	virtual void FocusChatChannel(USocialUser& InChannelUser);
 	virtual void FocusChatChannel(USocialChatChannel& InChannel);
+	virtual void DisplayChatChannel(USocialChatChannel& InChannel);
 
 	virtual TSubclassOf<USocialChatRoom> GetClassForChatRoom(ESocialChannelType Type) const;
 	virtual TSubclassOf<USocialChatChannel> GetClassForPrivateMessage() const { return USocialPrivateMessageChannel::StaticClass(); }
@@ -143,4 +146,5 @@ private:
 	mutable FOnChatChannelCreated OnChannelCreatedEvent;
 	mutable FOnChatChannelLeft OnChannelLeftEvent;
 	mutable FOnChatChannelFocusRequested OnChannelFocusRequestedEvent;
+	mutable FOnChatChannelDisplayRequested OnChannelDisplayRequestedEvent;
 };
\ No newline at end of file
