Added virtuals to SocialToolkit for children to be able to react when a particular list query has completed
Also deleted a few errant lines

diff --git a/Engine/Plugins/Online/OnlineFramework/Source/Party/Private/SocialToolkit.cpp b/Engine/Plugins/Online/OnlineFramework/Source/Party/Private/SocialToolkit.cpp
--- a/Engine/Plugins/Online/OnlineFramework/Source/Party/Private/SocialToolkit.cpp
+++ b/Engine/Plugins/Online/OnlineFramework/Source/Party/Private/SocialToolkit.cpp
@@ -637,6 +637,7 @@ void USocialToolkit::HandleReadFriendsListComplete(int32 LocalUserNum, bool bWas
 
 		//@todo DanH: This isn't actually quite correct - some of these could actually just be friend info for pending invites, not fully accepted friends. Should piece out the list into respective categories and process each separately (or make the associated event determination more complex)
 		ProcessUserList(FriendsList, SubsystemType, OnFriendshipEstablished());
+		OnQueryFriendsListSuccess(SubsystemType, FriendsList);
 	}
 	else
 	{
@@ -659,6 +660,7 @@ void USocialToolkit::HandleQueryBlockedPlayersComplete(const FUniqueNetId& UserI
 			TArray<TSharedRef<FOnlineBlockedPlayer>> BlockedPlayers;
 			FriendsInterface->GetBlockedPlayers(UserId, BlockedPlayers);
 			ProcessUserList(BlockedPlayers, SubsystemType, OnUserBlocked());
+			OnQueryBlockedPlayersSuccess(SubsystemType, BlockedPlayers);
 		}
 		else
 		{
@@ -681,6 +683,7 @@ void USocialToolkit::HandleQueryRecentPlayersComplete(const FUniqueNetId& UserId
 			TArray<TSharedRef<FOnlineRecentPlayer>> RecentPlayers;
 			FriendsInterface->GetRecentPlayers(UserId, Namespace, RecentPlayers);
 			ProcessUserList(RecentPlayers, SubsystemType, OnRecentPlayerAdded());
+			OnQueryRecentPlayersSuccess(SubsystemType, RecentPlayers);
 		}
 		else
 		{
diff --git a/Engine/Plugins/Online/OnlineFramework/Source/Party/Private/User/SocialUser.cpp b/Engine/Plugins/Online/OnlineFramework/Source/Party/Private/User/SocialUser.cpp
--- a/Engine/Plugins/Online/OnlineFramework/Source/Party/Private/User/SocialUser.cpp
+++ b/Engine/Plugins/Online/OnlineFramework/Source/Party/Private/User/SocialUser.cpp
@@ -871,7 +871,6 @@ bool USocialUser::HasSentPartyInvite(const FOnlinePartyTypeId& PartyTypeId) cons
 		const FUniqueNetIdRepl UserId = GetUserId(ESocialSubsystem::Primary);
 		if (ensure(LocalUserId.IsValid()) && UserId.IsValid())
 		{
-			// No advertised party info, check to see if this user has sent an invite
 			TArray<TSharedRef<IOnlinePartyJoinInfo>> AllPendingInvites;
 			if (PartyInterface->GetPendingInvites(*LocalUserId, AllPendingInvites))
 			{
diff --git a/Engine/Plugins/Online/OnlineFramework/Source/Party/Public/SocialToolkit.h b/Engine/Plugins/Online/OnlineFramework/Source/Party/Public/SocialToolkit.h
--- a/Engine/Plugins/Online/OnlineFramework/Source/Party/Public/SocialToolkit.h
+++ b/Engine/Plugins/Online/OnlineFramework/Source/Party/Public/SocialToolkit.h
@@ -108,14 +108,17 @@ protected:
 
 	virtual void OnOwnerLoggedIn();
 	virtual void OnOwnerLoggedOut();
-	virtual void OnAllUsersInitialized() {}
-
+	
 	virtual void NotifyFriendInviteFailed(const FUniqueNetId& InvitedUserId, const FString& InvitedUserName, ESendFriendInviteFailureReason FailureReason, bool bCanShow = true) {}
 
 	void QueryFriendsLists();
 	void QueryBlockedPlayers();
 	void QueryRecentPlayers();
 
+	virtual void OnQueryFriendsListSuccess(ESocialSubsystem SubsystemType, const TArray<TSharedRef<FOnlineFriend>>& FriendsList) {}
+	virtual void OnQueryBlockedPlayersSuccess(ESocialSubsystem SubsystemType, const TArray<TSharedRef<FOnlineBlockedPlayer>>& BlockedPlayers) {}
+	virtual void OnQueryRecentPlayersSuccess(ESocialSubsystem SubsystemType, const TArray<TSharedRef<FOnlineRecentPlayer>>& FriendsList) {}
+	
 	/** The type of SocialUser to create to represent known users */
 	TSubclassOf<USocialUser> SocialUserClass;
 
@@ -215,6 +218,5 @@ private:
 	mutable FOnRelationshipEstablished OnRecentPlayerAddedEvent;
 	
 	mutable FOnKnownUserInitialized OnKnownUserInitializedEvent;
-	mutable FBasicToolkitEvent OnFinishedStartupQueriesEvent;
 	mutable FBasicToolkitEvent OnToolkitResetEvent;
 };
