Moved modification of USocialManager::JoinedPartiesByTypeId before the InitializeParty() call in USocialManager::EstablishNewParty() so that it would be populated if the InitializeParty() call immediately and synchronously broadcasts a notification that the party has been initialized.

diff --git a/Engine/Plugins/Online/OnlineFramework/Source/Party/Private/SocialManager.cpp b/Engine/Plugins/Online/OnlineFramework/Source/Party/Private/SocialManager.cpp
--- a/Engine/Plugins/Online/OnlineFramework/Source/Party/Private/SocialManager.cpp
+++ b/Engine/Plugins/Online/OnlineFramework/Source/Party/Private/SocialManager.cpp
@@ -707,6 +707,9 @@ USocialParty* USocialManager::EstablishNewParty(const FUniqueNetId& LocalUserId,
 		NewParty->OnPartyLeaveBegin().AddUObject(this, &USocialManager::HandlePartyLeaveBegin, NewParty);
 		NewParty->OnPartyLeft().AddUObject(this, &USocialManager::HandlePartyLeft, NewParty);
 
+		// This must be done before InitializeParty(), as initialization can complete synchronously.
+		JoinedPartiesByTypeId.Add(PartyTypeId, NewParty);
+
 		NewParty->InitializeParty(OssParty.ToSharedRef());
 
 		if (NewParty->IsPersistentParty())
@@ -715,7 +718,6 @@ USocialParty* USocialManager::EstablishNewParty(const FUniqueNetId& LocalUserId,
 			HandlePersistentPartyStateChanged(NewParty->GetOssPartyState(), NewParty);
 		}
 
-		JoinedPartiesByTypeId.Add(PartyTypeId, NewParty);
 		return NewParty;
 	}
 
