[UE-68613] - Preventing local SocialUser init ensure on a missing unique ID when the user isn't logged in on that OSS

diff --git a/Engine/Plugins/Online/OnlineFramework/Source/Party/Private/User/SocialUser.cpp b/Engine/Plugins/Online/OnlineFramework/Source/Party/Private/User/SocialUser.cpp
--- a/Engine/Plugins/Online/OnlineFramework/Source/Party/Private/User/SocialUser.cpp
+++ b/Engine/Plugins/Online/OnlineFramework/Source/Party/Private/User/SocialUser.cpp
@@ -137,15 +137,18 @@ void USocialUser::InitLocalUser()
 		IOnlineSubsystem* OSS = OwningToolkit.GetSocialOss(SubsystemType);
 		check(OSS);
 
-		IOnlineIdentityPtr IdentityInterface = OSS->GetIdentityInterface();
-		FUniqueNetIdRepl LocalUserSubsystemId = IdentityInterface ? IdentityInterface->GetUniquePlayerId(OwningToolkit.GetLocalUserNum()) : FUniqueNetIdRepl();
-		if (ensure(LocalUserSubsystemId.IsValid()))
+		if (IOnlineIdentityPtr IdentityInterface = OSS->GetIdentityInterface())
 		{
-			SetSubsystemId(SubsystemType, LocalUserSubsystemId);
-		}
-		else
-		{
-			UE_LOG(LogParty, Error, TEXT("Local SocialUser unable to establish a valid UniqueId on subsystem [%s]"), ToString(SubsystemType));
+			const int32 LocalUserNum = OwningToolkit.GetLocalUserNum();
+			FUniqueNetIdRepl LocalUserSubsystemId = IdentityInterface->GetUniquePlayerId(LocalUserNum);
+			if (IdentityInterface->GetLoginStatus(LocalUserNum) == ELoginStatus::LoggedIn && ensure(LocalUserSubsystemId.IsValid()))
+			{
+				SetSubsystemId(SubsystemType, LocalUserSubsystemId);
+			}
+			else
+			{
+				UE_LOG(LogParty, Warning, TEXT("Local SocialUser unable to establish a valid UniqueId on subsystem [%s]"), ToString(SubsystemType));
+			}
 		}
 	}
 
