Slash Commands Engine

Added start of chat slash command support to new social module.
#tests win64 and preflight

diff --git a/Engine/Plugins/Online/OnlineFramework/Source/Party/Private/Chat/ChatSlashCommands.cpp b/Engine/Plugins/Online/OnlineFramework/Source/Party/Private/Chat/ChatSlashCommands.cpp
new file mode 100644
--- /dev/null
+++ b/Engine/Plugins/Online/OnlineFramework/Source/Party/Private/Chat/ChatSlashCommands.cpp
@@ -0,0 +1,428 @@
+// Copyright 1998-2017 Epic Games, Inc. All Rights Reserved.
+
+#include "Chat/ChatSlashCommands.h"
+#include "Chat/SocialChatManager.h"
+#include "SocialToolkit.h"
+#include "User/ISocialUserList.h"
+#include "User/SocialUserList.h"
+#include "SocialManager.h"
+
+
+#define LOCTEXT_NAMESPACE "ChatSlashCommandManager"
+
+FAutoCompleteStruct::FAutoCompleteStruct(const FString& InFullString, const TWeakPtr<const FChatSlashCommand>& InCmd, TWeakObjectPtr<USocialUser> InOptionalTargetUser)
+	: FullString(InFullString), SlashCommand(InCmd), OptionalTargetUser(InOptionalTargetUser)
+{
+	FRegisteredSlashCommands::TokenizeMessage(FullString, Tokens);
+}
+
+bool FRegisteredSlashCommands::TryExecuteCommandByMatchingText(const FString& UserTypedText)
+{
+	if (UserTypedText.StartsWith(TEXT("/")))
+	{
+		//find exact match in remaining complete strings 
+		for (TSharedPtr<FAutoCompleteStruct>& AutoCompleteDatum : AutoCompleteData)
+		{
+			if (AutoCompleteDatum.IsValid())
+			{
+				if (AutoCompleteDatum->FullString == UserTypedText)
+				{
+					if (AutoCompleteDatum->SlashCommand.IsValid())
+					{
+						AutoCompleteDatum->SlashCommand.Pin()->ExecuteSlashCommand( AutoCompleteDatum->OptionalTargetUser.Get() );
+						return true;
+					}
+				}
+			}
+		}
+	}
+
+	return false;
+}
+
+void FRegisteredSlashCommands::RegisterCommand(const TSharedPtr<FChatSlashCommand>& NewSlashCommand)
+{
+	if (NewSlashCommand.IsValid())
+	{
+		for (const TSharedPtr<FChatSlashCommand> Cmd : RegisteredCustomSlashCommands)
+		{
+			if (Cmd.IsValid())
+			{
+				if (Cmd->GetCommandNameString() == NewSlashCommand->GetCommandNameString())
+				{
+					UE_LOG(LogParty, Warning, TEXT("Attempting to register duplicate slash command"));
+					return;
+				}
+			}
+		}
+		for (const TSharedPtr<FChatSlashCommand> Cmd : RegisteredInteractionSlashCommands)
+		{
+			if (Cmd.IsValid())
+			{
+				if (Cmd->GetCommandNameString() == NewSlashCommand->GetCommandNameString())
+				{
+					UE_LOG(LogParty, Warning, TEXT("Attempting to register duplicate slash command"));
+					return;
+				}
+			}
+		}
+
+		if (USocialToolkit* Toolkit = MyToolkit.Get())
+		{
+			NewSlashCommand->Init(*Toolkit);
+		}
+
+		RegisteredCustomSlashCommands.Add(NewSlashCommand);
+	}
+}
+
+bool FRegisteredSlashCommands::IsEnabled()
+{
+	if (USocialToolkit* Toolkit = MyToolkit.Get())
+	{
+		return Toolkit->GetChatManager().AreSlashCommandsEnabled();
+	}
+	return false;
+}
+
+bool FRegisteredSlashCommands::NotifyUserTextChanged(const FText& InText) 
+{
+	FString InTextAsString = InText.ToString();
+	TArray<FString> Tokens;
+	TokenizeMessage(InTextAsString, Tokens);
+
+	if (InTextAsString.Len() > 0 && InTextAsString[0] == TEXT("/")[0])
+	{
+		AutoCompleteData.Reset();
+		PrepareInteractionAutocompleteStrings(Tokens);
+		for (const TSharedPtr<FChatSlashCommand>& Cmd : RegisteredCustomSlashCommands)
+		{
+			Cmd->GetAutoCompleteStrings(AutoCompleteData, Tokens);
+		}
+	}
+	else if (InTextAsString.Len() == 0)
+	{
+		AutoCompleteData.Reset();
+	}
+
+	if (SpaceWasJustTyped(InTextAsString))
+	{
+		//attempt spacebar execution on exact matches to autocomplete data
+		for (TSharedPtr<FAutoCompleteStruct>& AutoCompleteDatum : AutoCompleteData)
+		{
+			if (AutoCompleteDatum)
+			{
+				if (TokensExactMatch(AutoCompleteDatum->Tokens, Tokens))
+				{
+					TSharedPtr<const FChatSlashCommand> Cmd = AutoCompleteDatum->SlashCommand.Pin();
+					if (Cmd->HasSpacebarExecuteFunctionality())
+					{
+						if (Cmd->RequiresUserForExecution())
+						{
+							if (!AutoCompleteDatum->OptionalTargetUser.IsValid())
+							{
+								Cmd->ExecuteSlashCommand(AutoCompleteDatum->OptionalTargetUser.Get());
+								return true;
+							}
+						}
+						else
+						{
+							Cmd->ExecuteSlashCommand(AutoCompleteDatum->OptionalTargetUser.Get());
+							return true;
+						}
+					}
+				}
+			}
+		}
+		//attempt spacebar execution based on partial completion
+		for (const TSharedPtr<FChatSlashCommand>& Cmd : RegisteredCustomSlashCommands)
+		{
+			if (Cmd->CanExecuteSpacebarFromPartialTokens(Tokens))
+			{
+				Cmd->ExecuteSlashCommand(nullptr);
+				return true;
+			}
+		}
+	}
+
+	return false;
+}
+
+void FRegisteredSlashCommands::Init(USocialToolkit& Toolkit)
+{
+	const TArray<FSocialInteractionHandle> RegisteredInteractions = USocialManager::GetRegisteredInteractions();
+	for (FSocialInteractionHandle Interaction : RegisteredInteractions)
+	{
+		if (Interaction.GetSlashCommandToken().Len() > 0)
+		{
+			TSharedPtr<FInteractionCommandWrapper> SlashCommand = MakeShared<FInteractionCommandWrapper>(Interaction);
+			SlashCommand->Init(Toolkit);
+			RegisteredInteractionSlashCommands.Add(SlashCommand);
+		}
+	}
+
+	FInternationalization::Get().OnCultureChanged().AddSP(SharedThis(this), &FRegisteredSlashCommands::HandleCultureChanged);
+
+	MyToolkit = &Toolkit;
+}
+
+void FRegisteredSlashCommands::TokenizeMessage(const FString& InChatText, TArray<FString>& Tokens)
+{
+	//this is simple, but wrapping in a method provides consistent behavior for how slash commands will be tokenized
+	InChatText.ParseIntoArray(Tokens, TEXT(" "), true);
+}
+
+bool FRegisteredSlashCommands::TokensExactMatch(TArray<FString>& TokensLHS, TArray<FString>& TokensRHS)
+{
+	if (TokensRHS.Num() != TokensLHS.Num())
+	{
+		return false;
+	}
+
+	int32 NumTokens = TokensLHS.Num();
+
+	for (int Token = 0; Token < NumTokens; ++Token)
+	{
+		if (TokensLHS[Token] != TokensRHS[Token])
+		{
+			return false;
+		}
+	}
+	return true;
+}
+
+bool FRegisteredSlashCommands::CmdMatchesFirstToken(const FString& CmdString, const TArray<FString>& Tokens)
+{
+	int32 NumTokens = Tokens.Num();
+
+	if (NumTokens == 0)
+	{
+		return false;
+	}
+	else if (NumTokens == 1)
+	{
+		//old system used "Contains"
+		return CmdString.StartsWith(Tokens[0]);
+	}
+	else
+	{
+		//more than one token means first token should match exactly
+		return Tokens[0] == CmdString;
+	}
+}
+
+void FRegisteredSlashCommands::PrepareInteractionAutocompleteStrings(const TArray<FString>& UserTextTokens) 
+{
+	if (USocialToolkit* Toolkit = MyToolkit.Get())
+	{
+		if (UserTextTokens.Num() != 0)
+		{
+			TArray<TSharedPtr<FInteractionCommandWrapper>> RelevantCommands;
+			for (const TSharedPtr<FInteractionCommandWrapper>& InteractionCmd : RegisteredInteractionSlashCommands)
+			{
+				if (CmdMatchesFirstToken(InteractionCmd->GetCommandNameString(), UserTextTokens))
+				{
+					RelevantCommands.Add(InteractionCmd);
+				}
+			}
+
+			//only add autocomplete with user names after first characters have been typed and we've narrowed down the list to a single command
+			if (RelevantCommands.Num() == 1)
+			{
+				//only 1 viable command, so auto compete user names
+				const TSharedPtr<FInteractionCommandWrapper>& OnlyViableCmd = RelevantCommands[0];
+				
+				//these checks are expensive, only do once and filter down based on name.
+				if (!bValidUsersCached)
+				{					
+					for (USocialUser* User : Toolkit->GetAllUsers())
+					{
+						if (User)
+						{
+							OnlyViableCmd->TryCacheValidAutoCompleteUser(*User, UserTextTokens);
+						}
+					}
+					bValidUsersCached = true;
+				}
+				OnlyViableCmd->GetAutoCompleteStrings(AutoCompleteData, UserTextTokens);
+			}
+			else
+			{
+				//clear stale users while narrowing down commands to single command
+				if (bValidUsersCached)
+				{
+					for (const TSharedPtr<FInteractionCommandWrapper>& InteractionCmd : RelevantCommands)
+					{
+						InteractionCmd->ResetUserCache();
+					}
+				}
+				bValidUsersCached = false;
+				
+				//multiple commands, only autocomplete command names (ie first token)
+				for (const TSharedPtr<FInteractionCommandWrapper>& InteractionCmd : RelevantCommands)
+				{
+					InteractionCmd->GetAutoCompleteStrings(AutoCompleteData, UserTextTokens);
+				}
+			}
+		}
+	}
+	else
+	{
+		UE_LOG(LogParty, Warning, TEXT("Warning: Slash Command Component not initialized with a toolkit. Slash commands will be disabled until intialization."));
+	}
+}
+
+void FRegisteredSlashCommands::HandleCultureChanged() const
+{
+	for (const TSharedPtr<FChatSlashCommand>& CustomCmd : RegisteredCustomSlashCommands)
+	{
+		CustomCmd->RecacheForLocalization();
+	}
+	for (const TSharedPtr<FChatSlashCommand>& InteractionCmd : RegisteredInteractionSlashCommands)
+	{
+		InteractionCmd->RecacheForLocalization();
+	}
+}
+
+bool FRegisteredSlashCommands::SpaceWasJustTyped(const FString& NewUserText)
+{
+	bool bSpaceTyped = false;
+	uint32 NewLen = NewUserText.Len();
+
+	bool bUserAddedChar = NewLen > LastQueryTextLen;
+	if (bUserAddedChar && NewLen != 0)
+	{
+		bSpaceTyped = NewUserText[NewLen - 1] == TEXT(" ")[0];
+	}
+
+	//update internal state for next query; if text is cleared NewLen will be written to 1 when / is typed
+	LastQueryTextLen = NewLen;
+
+	return bSpaceTyped;
+}
+
+///////////////////////////////////////////////////////////////////////////////////////////////////////////////////////////////////
+
+void FChatSlashCommand::RecacheForLocalization() const
+{
+	//re-read the source text reference to get the new localization
+	CommandNameString = CommandNameTextSrc->ToString();
+}
+
+//////////////////////////////////////////////////////////////////////////////////
+// Chat Slash Command
+
+FChatSlashCommand::FChatSlashCommand(const FText& InCommandName) 
+	: CommandNameTextSrc(&InCommandName)
+{
+	//done outside of initializer list to avoid pesky headerfile order definition compile errors/warnings
+	CommandNameString = InCommandName.ToString();
+}
+
+void FChatSlashCommand::Init(USocialToolkit& InToolkit)
+{
+	MyToolkit = &InToolkit;
+}
+
+void FChatSlashCommand::GetAutoCompleteStrings(TArray< TSharedPtr<FAutoCompleteStruct> >& OutStringContainer, const TArray<FString>& UserTextTokens) const
+{
+	//default behavior is simply checking if first token match command name; override this for more complex behavior.
+	if (FRegisteredSlashCommands::CmdMatchesFirstToken(CommandNameString, UserTextTokens) && IsEnabled())
+	{
+		OutStringContainer.Add(MakeShared<FAutoCompleteStruct>(CommandNameString, SharedThis(this), nullptr));
+	}
+}
+
+#undef LOCTEXT_NAMESPACE
+FInteractionCommandWrapper::FInteractionCommandWrapper(FSocialInteractionHandle Interaction) 
+	: FChatSlashCommand(FText::FromString(Interaction.GetSlashCommandToken())),
+	WrappedInteraction(Interaction)
+{
+	CacheStringDataForLocalization();
+}
+
+void FInteractionCommandWrapper::ExecuteSlashCommand(USocialUser* OptionalTargetUser) const
+{
+	if (OptionalTargetUser)
+	{
+		WrappedInteraction.ExecuteInteraction(*OptionalTargetUser);
+	}
+}
+
+void FInteractionCommandWrapper::ResetUserCache()
+{
+	CachedValidUsers.Reset();
+}
+
+void FInteractionCommandWrapper::TryCacheValidAutoCompleteUser(USocialUser& User, const TArray<FString>& UserTextTokens)
+{	
+	//if user has typed partial name, check that name matches user; first token is command token
+	if (UserTextTokens.Num() > 1)
+	{
+		const FString& TypedUserName = UserTextTokens.Last();
+		if (!User.GetDisplayName().StartsWith(TypedUserName))
+		{
+			//player is not typing this user name, early out
+			return;
+		}
+	}
+
+	if (WrappedInteraction.IsAvailable(User))
+	{
+		CachedValidUsers.Add(&User);
+	}
+}
+
+void FInteractionCommandWrapper::GetAutoCompleteStrings(TArray< TSharedPtr<FAutoCompleteStruct> >& OutStringContainer, const TArray<FString>& UserTextTokens) const
+{
+	if(FRegisteredSlashCommands::CmdMatchesFirstToken(GetCommandNameString(), UserTextTokens))
+	{
+		if (CachedValidUsers.Num() != 0)
+		{
+			for (TWeakObjectPtr<USocialUser> CachedUser : CachedValidUsers)
+			{
+				if (CachedUser.IsValid() && WrappedInteraction.IsAvailable(*CachedUser))
+				{
+					FString AutoCompleteStringForUser = GetCommandNameString() + TEXT(" ") + CachedUser->GetDisplayName();
+					OutStringContainer.Add(MakeShared<FAutoCompleteStruct>(AutoCompleteStringForUser, SharedThis(this), CachedUser));
+				}
+			}
+		}
+		else
+		{
+			//no users checks yet, just populate command name until user data available for auto-complete
+			OutStringContainer.Add(MakeShared<FAutoCompleteStruct>(GetCommandNameString(), SharedThis(this), nullptr));
+		}
+	}
+}
+
+void FInteractionCommandWrapper::RecacheForLocalization() const
+{
+	//don't call super, we're going to custom process interaction names (eg remove spaces)
+	CacheStringDataForLocalization();
+}
+
+void FInteractionCommandWrapper::CacheStringDataForLocalization() const
+{
+	if (WrappedInteraction.IsValid())
+	{
+		//this will be up-to-date with current set localization
+		CommandNameString = TEXT("/") + WrappedInteraction.GetSlashCommandToken();
+
+		//compress the localized slash command name into a single token by removing spaces
+		TArray<FString> TokensToCompress; 
+		FRegisteredSlashCommands::TokenizeMessage(CommandNameString, TokensToCompress);
+		FString TokensCompressed = TEXT("");
+		for (FString& Token : TokensToCompress)
+		{
+			TokensCompressed += Token;
+		}
+		CommandNameString= TokensCompressed;
+	}
+	else
+	{
+		//UE_LOG(LogParty, Warning, "null slash command interaction detected, FInteractionCommandWrapper::RecacheForLocalization");
+	}
+}
+
+
diff --git a/Engine/Plugins/Online/OnlineFramework/Source/Party/Private/Chat/CommonSlashCommands.cpp b/Engine/Plugins/Online/OnlineFramework/Source/Party/Private/Chat/CommonSlashCommands.cpp
new file mode 100644
--- /dev/null
+++ b/Engine/Plugins/Online/OnlineFramework/Source/Party/Private/Chat/CommonSlashCommands.cpp
@@ -0,0 +1,128 @@
+// Copyright 1998-2017 Epic Games, Inc. All Rights Reserved.
+
+#include "Chat/CommonSlashCommands.h"
+#include "Chat/SocialChatChannel.h"
+#include "Chat/SocialChatManager.h"
+#include "SocialToolkit.h"
+
+/////////////////////////////////////////////////////////////////////////////////////////////////////////////////////
+// Channel Changing Slash Command
+FChannelChangeSlashCommand::FChannelChangeSlashCommand(const FText& CommandText, ESocialChannelType InChannelType) :
+	FChatSlashCommand(CommandText),
+	ChannelType(InChannelType)
+{
+}
+
+void FChannelChangeSlashCommand::ExecuteSlashCommand(USocialUser* OptionalTargetUser) const
+{
+	if (USocialToolkit* Tookit = GetToolkit())
+	{
+		if (USocialChatChannel* Channel = Tookit->GetChatManager().GetChatRoomForType(ChannelType))
+		{
+			Tookit->GetChatManager().FocusChatChannel(*Channel);
+		}
+	}
+}
+
+bool FChannelChangeSlashCommand::IsEnabled() const 
+{
+	if (USocialToolkit* Tookit = GetToolkit())
+	{
+		if (USocialChatChannel* Channel = Tookit->GetChatManager().GetChatRoomForType(ChannelType))
+		{
+			return true;
+		}
+	}
+	return false;
+}
+
+bool FChannelChangeSlashCommand::CanExecuteSpacebarFromPartialTokens(const TArray<FString>& UserTextTokens) const
+{
+	if (UserTextTokens.Num() == 1 && GetCommandNameString().StartsWith(UserTextTokens[0]))
+	{
+		return true;
+	}
+	return false;
+}
+
+
+////////////////////////////////////////////////////////////////////////////////////////////////////////////////////
+
+FPartyChannelSlashCommand::FPartyChannelSlashCommand() :
+	FChannelChangeSlashCommand(NSLOCTEXT("SlashCommands", "PartySlashCommand", "/party"), ESocialChannelType::Party)
+{
+}
+
+FGlobalChannelSlashCommand::FGlobalChannelSlashCommand() :
+	FChannelChangeSlashCommand(NSLOCTEXT("SlashCommands", "GlobalSlashCommand", "/global"), ESocialChannelType::General)
+{
+}
+
+FTeamChannelSlashCommand::FTeamChannelSlashCommand() :
+	FChannelChangeSlashCommand(NSLOCTEXT("SlashCommands", "TeamSlashCommand", "/team"), ESocialChannelType::Team)
+{
+}
+
+FFounderChannelSlashCommand::FFounderChannelSlashCommand() :
+	FChannelChangeSlashCommand(NSLOCTEXT("SlashCommands", "FounderSlashCommand", "/founder"), ESocialChannelType::Founder)
+{
+}
+///////////////////////////////////////////////////////////////////////////////////////////////////////////////////
+
+
+FReplySlashCommand::FReplySlashCommand() :
+	FChatSlashCommand(NSLOCTEXT("SlashCommands", "ReplySlashCommand", "/reply"))
+{}
+
+void FReplySlashCommand::Init(USocialToolkit& InToolkit)
+{
+	FChatSlashCommand::Init(InToolkit);
+
+	USocialChatManager& ChatManager = InToolkit.GetChatManager();
+	ChatManager.OnChannelCreated().AddSP(this, &FReplySlashCommand::HandleChannelCreated);
+	ChatManager.OnChannelLeft().AddSP(this, &FReplySlashCommand::HandleChannelLeft);
+}
+
+void FReplySlashCommand::ExecuteSlashCommand(USocialUser* OptionalTargetUser) const 
+{
+
+	if (USocialToolkit* Tookit = GetToolkit())
+	{
+		if (LastUserChannel.IsValid())
+		{
+			Tookit->GetChatManager().FocusChatChannel(*LastUserChannel);
+		}
+	}
+}
+
+bool FReplySlashCommand::IsEnabled() const 
+{
+	return this->LastUserChannel != nullptr;
+}
+
+void FReplySlashCommand::HandleChannelCreated(USocialChatChannel& NewChannel)
+{
+	if (ESocialChannelType::Private == NewChannel.GetChannelType())
+	{
+		TWeakPtr<FChatSlashCommand> WeakThisPtr = this->AsShared();
+		NewChannel.OnMessageReceived().AddLambda([WeakThisPtr, &NewChannel, this](FSocialChatMessageRef Messsage) {
+			if (WeakThisPtr.IsValid())
+			{
+ 				this->LastUserChannel = &NewChannel;
+			}
+		});
+	}
+}
+
+void FReplySlashCommand::HandleChannelLeft(USocialChatChannel& LeavingChannel)
+{
+	LeavingChannel.OnMessageReceived().RemoveAll(this);
+
+	if (LastUserChannel.IsValid())
+	{
+		if (LastUserChannel == &LeavingChannel)
+		{
+			LastUserChannel = nullptr;
+		}
+	}
+}
diff --git a/Engine/Plugins/Online/OnlineFramework/Source/Party/Private/Interactions/CoreInteractions.cpp b/Engine/Plugins/Online/OnlineFramework/Source/Party/Private/Interactions/CoreInteractions.cpp
--- a/Engine/Plugins/Online/OnlineFramework/Source/Party/Private/Interactions/CoreInteractions.cpp
+++ b/Engine/Plugins/Online/OnlineFramework/Source/Party/Private/Interactions/CoreInteractions.cpp
@@ -28,7 +28,7 @@ FText FSocialInteraction_AddFriend::GetDisplayName(const USocialUser& User)
 
 FString FSocialInteraction_AddFriend::GetSlashCommandToken()
 {
-	return TEXT("friend");
+	return FString();
 }
 
 bool FSocialInteraction_AddFriend::CanExecute(const USocialUser& User)
@@ -90,7 +90,7 @@ FText FSocialInteraction_RemoveFriend::GetDisplayName(const USocialUser& User)
 
 FString FSocialInteraction_RemoveFriend::GetSlashCommandToken()
 {
-	return TEXT("unfriend");
+	return LOCTEXT("SlashCommand_Unfriend", "unfriend").ToString();
 }
 
 bool FSocialInteraction_RemoveFriend::CanExecute(const USocialUser& User)
@@ -115,7 +115,7 @@ FText FSocialInteraction_AcceptFriendInvite::GetDisplayName(const USocialUser& U
 
 FString FSocialInteraction_AcceptFriendInvite::GetSlashCommandToken()
 {
-	return TEXT("accept");
+	return FString();
 }
 
 bool FSocialInteraction_AcceptFriendInvite::CanExecute(const USocialUser& User)
@@ -139,7 +139,7 @@ FText FSocialInteraction_RejectFriendInvite::GetDisplayName(const USocialUser& U
 
 FString FSocialInteraction_RejectFriendInvite::GetSlashCommandToken()
 {
-	return TEXT("reject");
+	return FString();
 }
 
 bool FSocialInteraction_RejectFriendInvite::CanExecute(const USocialUser& User)
@@ -163,7 +163,7 @@ FText FSocialInteraction_Block::GetDisplayName(const USocialUser& User)
 
 FString FSocialInteraction_Block::GetSlashCommandToken()
 {
-	return TEXT("block");
+	return LOCTEXT("SlashCommand_BlockUser", "block").ToString();
 }
 
 bool FSocialInteraction_Block::CanExecute(const USocialUser& User)
@@ -187,7 +187,7 @@ FText FSocialInteraction_Unblock::GetDisplayName(const USocialUser& User)
 
 FString FSocialInteraction_Unblock::GetSlashCommandToken()
 {
-	return TEXT("unblock");
+	return LOCTEXT("SlashCommand_Unblock", "unblock").ToString();
 }
 
 bool FSocialInteraction_Unblock::CanExecute(const USocialUser& User)
@@ -212,7 +212,7 @@ FText FSocialInteraction_PrivateMessage::GetDisplayName(const USocialUser& User)
 
 FString FSocialInteraction_PrivateMessage::GetSlashCommandToken()
 {
-	return TEXT("whisper");
+	return LOCTEXT("SlashCommand_PrivateMessage", "whisper").ToString();
 }
 
 bool FSocialInteraction_PrivateMessage::CanExecute(const USocialUser& User)
@@ -241,7 +241,7 @@ FText FSocialInteraction_ShowPlatformProfile::GetDisplayName(const USocialUser&
 
 FString FSocialInteraction_ShowPlatformProfile::GetSlashCommandToken()
 {
-	return TEXT("profile");
+	return TEXT("");
 }
 
 bool FSocialInteraction_ShowPlatformProfile::CanExecute(const USocialUser& User)
diff --git a/Engine/Plugins/Online/OnlineFramework/Source/Party/Private/Interactions/PartyInteractions.cpp b/Engine/Plugins/Online/OnlineFramework/Source/Party/Private/Interactions/PartyInteractions.cpp
--- a/Engine/Plugins/Online/OnlineFramework/Source/Party/Private/Interactions/PartyInteractions.cpp
+++ b/Engine/Plugins/Online/OnlineFramework/Source/Party/Private/Interactions/PartyInteractions.cpp
@@ -21,7 +21,7 @@ FText FSocialInteraction_InviteToParty::GetDisplayName(const USocialUser& User)
 
 FString FSocialInteraction_InviteToParty::GetSlashCommandToken()
 {
-	return TEXT("invite");
+	return (LOCTEXT("SlashCommand_InviteToParty", "invite")).ToString();
 }
 
 bool FSocialInteraction_InviteToParty::CanExecute(const USocialUser& User)
@@ -45,7 +45,7 @@ FText FSocialInteraction_JoinParty::GetDisplayName(const USocialUser& User)
 
 FString FSocialInteraction_JoinParty::GetSlashCommandToken()
 {
-	return TEXT("join");
+	return LOCTEXT("SlashCommand_JoinParty", "join").ToString();
 }
 
 bool FSocialInteraction_JoinParty::CanExecute(const USocialUser& User)
@@ -70,7 +70,8 @@ FText FSocialInteraction_AcceptPartyInvite::GetDisplayName(const USocialUser& Us
 
 FString FSocialInteraction_AcceptPartyInvite::GetSlashCommandToken()
 {
-	return TEXT("acceptpartyinvite");
+	//join should be the preferred method of accepting a party invite
+	return FString();
 }
 
 bool FSocialInteraction_AcceptPartyInvite::CanExecute(const USocialUser& User)
@@ -94,7 +95,7 @@ FText FSocialInteraction_RejectPartyInvite::GetDisplayName(const USocialUser& Us
 
 FString FSocialInteraction_RejectPartyInvite::GetSlashCommandToken()
 {
-	return TEXT("rejectpartyinvite");
+	return FString();
 }
 
 bool FSocialInteraction_RejectPartyInvite::CanExecute(const USocialUser& User)
@@ -118,7 +119,7 @@ FText FSocialInteraction_LeaveParty::GetDisplayName(const USocialUser& User)
 
 FString FSocialInteraction_LeaveParty::GetSlashCommandToken()
 {
-	return TEXT("leave");
+	return LOCTEXT("SlashCommand_LeaveParty", "leave").ToString();
 }
 
 bool FSocialInteraction_LeaveParty::CanExecute(const USocialUser& User)
@@ -150,7 +151,7 @@ FText FSocialInteraction_KickPartyMember::GetDisplayName(const USocialUser& User
 
 FString FSocialInteraction_KickPartyMember::GetSlashCommandToken()
 {
-	return TEXT("kick");
+	return LOCTEXT("SlashCommand_KickMember", "kick").ToString();
 }
 
 bool FSocialInteraction_KickPartyMember::CanExecute(const USocialUser& User)
@@ -178,7 +179,7 @@ FText FSocialInteraction_PromoteToPartyLeader::GetDisplayName(const USocialUser&
 
 FString FSocialInteraction_PromoteToPartyLeader::GetSlashCommandToken()
 {
-	return TEXT("promote");
+	return LOCTEXT("SlashCommand_PromoteToLeader", "promote").ToString();
 }
 
 bool FSocialInteraction_PromoteToPartyLeader::CanExecute(const USocialUser& User)
diff --git a/Engine/Plugins/Online/OnlineFramework/Source/Party/Public/Chat/ChatSlashCommands.h b/Engine/Plugins/Online/OnlineFramework/Source/Party/Public/Chat/ChatSlashCommands.h
new file mode 100644
--- /dev/null
+++ b/Engine/Plugins/Online/OnlineFramework/Source/Party/Public/Chat/ChatSlashCommands.h
@@ -0,0 +1,148 @@
+// Copyright 1998-2018 Epic Games, Inc. All Rights Reserved.
+
+#pragma once
+//#include "UObject/Interface.h"
+
+#include "UObject/Object.h"
+#include "Templates/SharedPointer.h"
+#include "Internationalization/Text.h"
+#include "Engine/LocalPlayer.h"
+#include "Interactions/SocialInteractionHandle.h"
+
+#include "Internationalization/Internationalization.h"
+#include "Internationalization/CulturePointer.h"
+
+class FSocialUserList;
+class FChatSlashCommand;
+class FInteractionCommandWrapper;
+class ULocalPlayer;
+class USocialManager;
+class USocialChatManager;
+class USocialUser;
+class USocialToolkit;
+
+
+
+///////////////////////////////////////////////////////////////////////////////
+struct FAutoCompleteStruct
+{
+	FAutoCompleteStruct(const FString& InFullString, const TWeakPtr<const FChatSlashCommand>& InCmd, TWeakObjectPtr<USocialUser> InOptionalTargetUser);
+
+	//Cacheing data with strings is tricky because they will become invalid when the user changes their localization
+	//but a user can't change this while typing a command, so it is okay in this case.
+	FString FullString;
+	TWeakPtr<const FChatSlashCommand> SlashCommand;
+	TWeakObjectPtr<USocialUser> OptionalTargetUser;
+	TArray<FString> Tokens;
+};
+
+///////////////////////////////////////////////////////////////////////////////
+//Slash Command Component
+class PARTY_API FRegisteredSlashCommands : public TSharedFromThis<FRegisteredSlashCommands>
+{
+public:
+	static void TokenizeMessage(const FString& InChatText, TArray<FString>& Tokens);
+	static bool TokensExactMatch(TArray<FString>& TokensLHS, TArray<FString>& TokensRHS);
+	static bool CmdMatchesFirstToken(const FString& CmdString, const TArray<FString>& Tokens);
+
+	FRegisteredSlashCommands() = default;
+	void Init(USocialToolkit& Toolkit);
+
+	/** main entry point for class encapsulated behavior; returns true if command executed */
+	bool NotifyUserTextChanged(const FText& InText);
+
+	bool TryExecuteCommandByMatchingText(const FString& UserTypedText);
+	bool HasAutoCompleteSuggestions(){ return AutoCompleteData.Num() != 0; }
+	const TArray<TSharedPtr<FAutoCompleteStruct>>& GetAutoCompleteStrings() const { return AutoCompleteData; }
+	void RegisterCommand(const TSharedPtr<FChatSlashCommand>& NewSlashCommand);
+	bool IsEnabled();
+
+private: 
+	void PrepareInteractionAutocompleteStrings(const TArray<FString>& StringTokens);
+	void HandleCultureChanged() const;
+
+	uint32 LastQueryTextLen = 0;
+	bool bValidUsersCached = false;
+	bool SpaceWasJustTyped(const FString& NewUserText);
+
+private: 
+	TArray<TSharedPtr<FChatSlashCommand>> RegisteredCustomSlashCommands;
+	TArray<TSharedPtr<FInteractionCommandWrapper>> RegisteredInteractionSlashCommands;
+
+	//once set, this should always be valid since lifetime of SocialManager is tied to game instance
+	TWeakObjectPtr<USocialToolkit> MyToolkit;
+	mutable TArray<TSharedPtr<FAutoCompleteStruct>> AutoCompleteData;
+
+	FRegisteredSlashCommands(const FRegisteredSlashCommands& copy) = delete;
+	FRegisteredSlashCommands(FRegisteredSlashCommands&& move) = delete;
+	FRegisteredSlashCommands& operator=(const FRegisteredSlashCommands& copy) = delete;
+	FRegisteredSlashCommands& operator=(FRegisteredSlashCommands&& move) = delete;
+};
+///////////////////////////////////////////////////////////////////////////////
+
+
+
+//////////////////////////////////////////////////////////////////////////////////
+class PARTY_API FChatSlashCommand : public TSharedFromThis<FChatSlashCommand>
+{
+public:
+	/** @param InCommandName Command name including / prefix. eg "/party".	*/
+	explicit FChatSlashCommand(const FText& InCommandName);
+	virtual ~FChatSlashCommand() = default;
+
+	virtual void Init(USocialToolkit& InToolkit);
+	virtual bool IsEnabled() const = 0;
+	virtual void ExecuteSlashCommand(USocialUser* OptionalTargetUser) const = 0;
+
+	virtual void GetAutoCompleteStrings(TArray< TSharedPtr<FAutoCompleteStruct> >& OutStringContainer, const TArray<FString>& UserTextTokens) const;
+	virtual bool CanExecuteSpacebarFromPartialTokens(const TArray<FString>& UserTextTokens) const { return false;  }
+	virtual bool HasSpacebarExecuteFunctionality() const { return false; }
+	virtual bool RequiresUserForExecution() const {return false;}
+	virtual void RecacheForLocalization() const;
+
+	const FString& GetCommandNameString() const { return CommandNameString; }
+
+protected:
+	USocialToolkit* GetToolkit() const { return MyToolkit.Get(); }
+	mutable FString CommandNameString;
+
+private:
+	const FText* CommandNameTextSrc;
+	TWeakObjectPtr<USocialToolkit> MyToolkit;
+
+	FChatSlashCommand(const FChatSlashCommand& copy) = delete;
+	FChatSlashCommand(FChatSlashCommand&& move) = delete;
+	FChatSlashCommand& operator=(const FChatSlashCommand& copy) = delete;
+	FChatSlashCommand& operator=(FChatSlashCommand&& move) = delete;
+};
+
+//////////////////////////////////////////////////////////////////////////////////
+
+class PARTY_API FInteractionCommandWrapper: public FChatSlashCommand
+{
+public:
+	/** Interaction tokens will have / prefix appended.*/
+	FInteractionCommandWrapper(FSocialInteractionHandle Interaction);
+
+	virtual void ExecuteSlashCommand(USocialUser* OptionalTargetUser) const override; 
+	virtual bool IsEnabled() const override	{ return true; 	}
+	virtual bool HasSpacebarExecuteFunctionality() const { return true; }
+	virtual bool RequiresUserForExecution() const override{ return true; }
+	virtual bool CanExecuteSpacebarFromPartialTokens(const TArray<FString>& UserTextTokens) const{return false;	}
+	
+	void ResetUserCache(); 
+	FORCEINLINE void TryCacheValidAutoCompleteUser(USocialUser& User, const TArray<FString>& StringTokens);
+	virtual void GetAutoCompleteStrings(TArray< TSharedPtr<FAutoCompleteStruct> >& OutStringContainer, const TArray<FString>& UserTextTokens) const override;
+	virtual void RecacheForLocalization() const override;
+
+private:
+	FSocialInteractionHandle WrappedInteraction;
+	mutable FString CachedCommandToken;
+	TArray<TWeakObjectPtr<USocialUser>> CachedValidUsers;
+
+	/* 
+	* NOTE: we cannot simply cache the FText long term because localization changes while running will invalidate cache.
+	* So, there exists the following function to re-query the localization
+	*/
+	void CacheStringDataForLocalization() const;
+};
\ No newline at end of file
diff --git a/Engine/Plugins/Online/OnlineFramework/Source/Party/Public/Chat/CommonSlashCommands.h b/Engine/Plugins/Online/OnlineFramework/Source/Party/Public/Chat/CommonSlashCommands.h
new file mode 100644
--- /dev/null
+++ b/Engine/Plugins/Online/OnlineFramework/Source/Party/Public/Chat/CommonSlashCommands.h
@@ -0,0 +1,72 @@
+// Copyright 1998-2018 Epic Games, Inc. All Rights Reserved.
+
+#pragma once
+
+#include "Chat/ChatSlashCommands.h"
+#include "Containers/UnrealString.h"
+#include "SocialTypes.h"
+#include "Chat/SocialChatChannel.h"
+
+class ULocalPlayer;
+class USocialUser;
+class USocialManager;
+class USocialToolkit;
+class USocialChatChannel;
+enum class ESocialChannelType : uint8;
+
+class PARTY_API FChannelChangeSlashCommand : public FChatSlashCommand
+{
+public:
+	FChannelChangeSlashCommand(const FText& CommandText, ESocialChannelType InChannelType);
+
+	virtual void ExecuteSlashCommand(USocialUser* OptionalTargetUser) const override;
+	virtual bool IsEnabled() const override;
+	virtual bool CanExecuteSpacebarFromPartialTokens(const TArray<FString>& UserTextTokens) const override;
+	virtual bool HasSpacebarExecuteFunctionality() const { return true; }
+
+private:
+	ESocialChannelType ChannelType;
+
+};
+
+class PARTY_API FPartyChannelSlashCommand : public FChannelChangeSlashCommand
+{
+public:
+	FPartyChannelSlashCommand();
+};
+
+class PARTY_API FGlobalChannelSlashCommand : public FChannelChangeSlashCommand
+{
+public:
+	FGlobalChannelSlashCommand();
+};
+
+class PARTY_API FTeamChannelSlashCommand : public FChannelChangeSlashCommand
+{
+public:
+	FTeamChannelSlashCommand();
+};
+
+class PARTY_API FFounderChannelSlashCommand : public FChannelChangeSlashCommand
+{
+public:
+	FFounderChannelSlashCommand();
+};
+
+/////////////////////////////////////////////////////////////////////////////////////////////////////////////
+
+class PARTY_API FReplySlashCommand: public FChatSlashCommand
+{
+public:
+	FReplySlashCommand();
+	virtual void Init(USocialToolkit& InToolkit) override;
+	virtual void ExecuteSlashCommand(USocialUser* OptionalTargetUser) const override;
+	virtual bool IsEnabled() const override;
+	TWeakObjectPtr<USocialChatChannel> LastUserChannel = nullptr;
+
+private:
+	void HandleChannelCreated(USocialChatChannel& NewChannel);
+	void HandleChannelLeft(USocialChatChannel& LeavingChannel);
+};
+
+
diff --git a/Engine/Plugins/Online/OnlineFramework/Source/Party/Public/Chat/SocialChatManager.cpp b/Engine/Plugins/Online/OnlineFramework/Source/Party/Public/Chat/SocialChatManager.cpp
--- a/Engine/Plugins/Online/OnlineFramework/Source/Party/Public/Chat/SocialChatManager.cpp
+++ b/Engine/Plugins/Online/OnlineFramework/Source/Party/Public/Chat/SocialChatManager.cpp
@@ -12,6 +12,7 @@
 #include "SocialManager.h"
 #include "Party/SocialParty.h"
 #include "SocialPartyChatRoom.h"
+#include "ChatSlashCommands.h"
 
 USocialChatRoom* USocialChatManager::GetChatRoom(const FChatRoomId& ChannelId) const
 {
@@ -111,6 +112,32 @@ void USocialChatManager::ExitChatRoom(const FChatRoomId& RoomId, ESocialSubsyste
 	}
 }
 
+void USocialChatManager::OnChannelCreatedInternal(USocialChatChannel& CreatedChannel)
+{
+	ESocialChannelType ChannelType = CreatedChannel.GetChannelType();
+	switch (ChannelType)
+	{
+		case ESocialChannelType::Founder:
+		case ESocialChannelType::General:
+		case ESocialChannelType::Party:
+		case ESocialChannelType::Team:
+		{
+			ChannelsByType.Add(ChannelType, &CreatedChannel);
+		}
+		break;
+	}
+
+	OnChannelCreated().Broadcast(CreatedChannel);
+}
+
+void USocialChatManager::OnChannelLeftInternal(USocialChatChannel& ChannelLeft)
+{
+	ESocialChannelType ChannelType = ChannelLeft.GetChannelType();
+	ChannelsByType.Remove(ChannelType);
+
+	OnChannelLeft().Broadcast(ChannelLeft);
+}
+
 USocialChatChannel& USocialChatManager::CreateChatChannel(USocialUser& InRecipient)
 {
 	return FindOrCreateChannel(InRecipient);
@@ -187,6 +214,15 @@ USocialToolkit& USocialChatManager::GetOwningToolkit() const
 	return *OwnerToolkit;
 }
 
+USocialChatChannel* USocialChatManager::GetChatRoomForType(ESocialChannelType Key)
+{
+	if (TWeakObjectPtr<USocialChatChannel>* Channel = ChannelsByType.Find(Key))
+	{
+		return Channel->Get();
+	}
+	return nullptr;
+}
+
 // This should only be possible for external code to do this via an ISocialUser directly
 // The user will hold a weak ptr to the direct message channel - if invalid, it'll ask the chat manager to make a channel for it
 // That way we still hit all the channel creation events but give the user fast access when sending whispers
@@ -260,7 +296,7 @@ USocialChatRoom& USocialChatManager::FindOrCreateRoom(const FChatRoomId& RoomId)
 
 	NewRoomChannel->Initialize(&GetOwningToolkit().GetLocalUser(), RoomId, ChannelType);
 	ChatRoomsById.Add(RoomId, NewRoomChannel);
-	OnChannelCreated().Broadcast(*NewRoomChannel);
+	OnChannelCreatedInternal(*NewRoomChannel);
 
 	return *NewRoomChannel;
 }
@@ -280,7 +316,7 @@ USocialChatChannel& USocialChatManager::FindOrCreateChannel(USocialUser& SocialU
 
 	NewPMChannel->Initialize(&SocialUser, FChatRoomId(TEXT("private")), ESocialChannelType::Private);
 	DirectChannelsByTargetUser.Add(&SocialUser, NewPMChannel);
-	OnChannelCreated().Broadcast(*NewPMChannel);
+	OnChannelCreatedInternal(*NewPMChannel);
 
 	return *NewPMChannel;
 }
@@ -300,7 +336,7 @@ USocialChatChannel& USocialChatManager::FindOrCreateChannel(const FText& Display
 
 	NewReadOnlyChannel->Initialize(&GetOwningToolkit().GetLocalUser(), DisplayName.ToString(), ESocialChannelType::General);
 	ReadOnlyChannelsByDisplayName.Add(DisplayName.ToString(), NewReadOnlyChannel);
-	OnChannelCreated().Broadcast(*NewReadOnlyChannel);
+	OnChannelCreatedInternal(*NewReadOnlyChannel);
 
 	return *NewReadOnlyChannel;
 }
@@ -361,7 +397,7 @@ void USocialChatManager::HandleChatRoomExit(const FUniqueNetId& LocalUserId, con
 		if (ensure(Room))
 		{
 			ChatRoomsById.Remove(RoomId);
-			OnChannelLeft().Broadcast(**Room);
+			OnChannelLeftInternal(**Room);
 		}
 	}
 	else
diff --git a/Engine/Plugins/Online/OnlineFramework/Source/Party/Public/Chat/SocialChatManager.h b/Engine/Plugins/Online/OnlineFramework/Source/Party/Public/Chat/SocialChatManager.h
--- a/Engine/Plugins/Online/OnlineFramework/Source/Party/Public/Chat/SocialChatManager.h
+++ b/Engine/Plugins/Online/OnlineFramework/Source/Party/Public/Chat/SocialChatManager.h
@@ -45,7 +45,7 @@ struct FSocialChatChannelConfig
 };
 
 /** The chat manager is a fully passive construct that watches for creation of chat rooms and message activity therein */
-UCLASS(Within=SocialToolkit)
+UCLASS(Within=SocialToolkit, Config=Game)
 class PARTY_API USocialChatManager : public UObject
 {
 	GENERATED_BODY()
@@ -104,6 +104,11 @@ public:
 
 	USocialToolkit& GetOwningToolkit() const;
 
+	bool AreSlashCommandsEnabled() { return bEnableChatSlashCommands; }
+
+	USocialChatChannel* GetChatRoomForType(ESocialChannelType Key);
+
+
 protected:
 	IOnlineChatPtr GetOnlineChatInterface(ESocialSubsystem InSocialSubsystem = ESocialSubsystem::Primary) const;
 	virtual void InitializeChatManager();
@@ -111,12 +116,16 @@ protected:
 
 	virtual void HandleChatRoomMessageReceived(const FUniqueNetId& LocalUserId, const FChatRoomId& RoomId, const TSharedRef<FChatMessage>& ChatMessage);
 	virtual void HandleChatPrivateMessageReceived(const FUniqueNetId& LocalUserId, const TSharedRef<FChatMessage>& ChatMessage);
+
+	virtual void OnChannelCreatedInternal(USocialChatChannel& CreatedChannel);
+	virtual void OnChannelLeftInternal(USocialChatChannel& ChannelLeft);
 private:
+	TMap < ESocialChannelType, TWeakObjectPtr<USocialChatChannel>> ChannelsByType;
 
 	USocialChatRoom& FindOrCreateRoom(const FChatRoomId& RoomId);
 	USocialChatChannel& FindOrCreateChannel(USocialUser& SocialUser);
 	USocialChatChannel& FindOrCreateChannel(const FText& DisplayName);
-	
+
 	void HandleChatRoomCreated(const FUniqueNetId& LocalUserId, const FChatRoomId& RoomId, bool bWasSuccessful, const FString& Error);
 	void HandleChatRoomConfigured(const FUniqueNetId& LocalUserId, const FChatRoomId& RoomId, bool bWasSuccessful, const FString& Error);
 	void HandleChatRoomJoinPublic(const FUniqueNetId& LocalUserId, const FChatRoomId& RoomId, bool bWasSuccessful, const FString& Error);
@@ -143,6 +152,9 @@ private:
 	UPROPERTY()
 	TMap<FString, USocialReadOnlyChatChannel*> ReadOnlyChannelsByDisplayName;
 
+	UPROPERTY(config)
+	bool bEnableChatSlashCommands = true;
+
 	mutable FOnChatChannelCreated OnChannelCreatedEvent;
 	mutable FOnChatChannelLeft OnChannelLeftEvent;
 	mutable FOnChatChannelFocusRequested OnChannelFocusRequestedEvent;
