Fix slash commands not clearing user cache.
Fix slash commands re-checking that users are valid even though they're in cached valid users container.
#tests PIE, compiled shipping client.

diff --git a/Engine/Plugins/Online/OnlineFramework/Source/Party/Private/Chat/ChatSlashCommands.cpp b/Engine/Plugins/Online/OnlineFramework/Source/Party/Private/Chat/ChatSlashCommands.cpp
--- a/Engine/Plugins/Online/OnlineFramework/Source/Party/Private/Chat/ChatSlashCommands.cpp
+++ b/Engine/Plugins/Online/OnlineFramework/Source/Party/Private/Chat/ChatSlashCommands.cpp
@@ -251,7 +251,7 @@ void FRegisteredSlashCommands::PrepareInteractionAutocompleteStrings(const TArra
 				//clear stale users while narrowing down commands to single command
 				if (bValidUsersCached)
 				{
-					for (const TSharedPtr<FInteractionCommandWrapper>& InteractionCmd : RelevantCommands)
+					for (const TSharedPtr<FInteractionCommandWrapper>& InteractionCmd : RegisteredInteractionSlashCommands)
 					{
 						InteractionCmd->ResetUserCache();
 					}
@@ -381,10 +381,21 @@ void FInteractionCommandWrapper::GetAutoCompleteStrings(TArray< TSharedPtr<FAuto
 		{
 			for (TWeakObjectPtr<USocialUser> CachedUser : CachedValidUsers)
 			{
-				if (CachedUser.IsValid() && WrappedInteraction.IsAvailable(*CachedUser))
+				if (CachedUser.IsValid())
 				{
-					FString AutoCompleteStringForUser = GetCommandNameString() + TEXT(" ") + CachedUser->GetDisplayName();
-					OutStringContainer.Add(MakeShared<FAutoCompleteStruct>(AutoCompleteStringForUser, SharedThis(this), CachedUser));
+					bool bUserMatchesText = true;
+
+					if (UserTextTokens.Num() >= 2)
+					{
+						const FString& TypedUserName = UserTextTokens.Last();
+						bUserMatchesText = CachedUser->GetDisplayName().StartsWith(TypedUserName);
+					}
+
+					if(bUserMatchesText)
+					{
+						FString AutoCompleteStringForUser = GetCommandNameString() + TEXT(" ") + CachedUser->GetDisplayName();
+						OutStringContainer.Add(MakeShared<FAutoCompleteStruct>(AutoCompleteStringForUser, SharedThis(this), CachedUser));
+					}
 				}
 			}
 		}
