Updating copyright notice

diff --git a/Engine/Plugins/Online/OnlineFramework/Source/Party/Private/Chat/ChatSlashCommands.cpp b/Engine/Plugins/Online/OnlineFramework/Source/Party/Private/Chat/ChatSlashCommands.cpp
--- a/Engine/Plugins/Online/OnlineFramework/Source/Party/Private/Chat/ChatSlashCommands.cpp
+++ b/Engine/Plugins/Online/OnlineFramework/Source/Party/Private/Chat/ChatSlashCommands.cpp
@@ -1,4 +1,4 @@
-// Copyright 1998-2017 Epic Games, Inc. All Rights Reserved.
+// Copyright 1998-2019 Epic Games, Inc. All Rights Reserved.
 
 #include "Chat/ChatSlashCommands.h"
 #include "Chat/SocialChatManager.h"
diff --git a/Engine/Plugins/Online/OnlineFramework/Source/Party/Private/Chat/CommonSlashCommands.cpp b/Engine/Plugins/Online/OnlineFramework/Source/Party/Private/Chat/CommonSlashCommands.cpp
--- a/Engine/Plugins/Online/OnlineFramework/Source/Party/Private/Chat/CommonSlashCommands.cpp
+++ b/Engine/Plugins/Online/OnlineFramework/Source/Party/Private/Chat/CommonSlashCommands.cpp
@@ -1,4 +1,4 @@
-// Copyright 1998-2017 Epic Games, Inc. All Rights Reserved.
+// Copyright 1998-2019 Epic Games, Inc. All Rights Reserved.
 
 #include "Chat/CommonSlashCommands.h"
 #include "Chat/SocialChatChannel.h"
diff --git a/Engine/Plugins/Online/OnlineFramework/Source/Party/Public/Chat/ChatSlashCommands.h b/Engine/Plugins/Online/OnlineFramework/Source/Party/Public/Chat/ChatSlashCommands.h
--- a/Engine/Plugins/Online/OnlineFramework/Source/Party/Public/Chat/ChatSlashCommands.h
+++ b/Engine/Plugins/Online/OnlineFramework/Source/Party/Public/Chat/ChatSlashCommands.h
@@ -1,4 +1,5 @@
-// Copyright 1998-2018 Epic Games, Inc. All Rights Reserved.
+// Copyright 1998-2019 Epic Games, Inc. All Rights Reserved.
+
 
 #pragma once
 //#include "UObject/Interface.h"
diff --git a/Engine/Plugins/Online/OnlineFramework/Source/Party/Public/Chat/CommonSlashCommands.h b/Engine/Plugins/Online/OnlineFramework/Source/Party/Public/Chat/CommonSlashCommands.h
--- a/Engine/Plugins/Online/OnlineFramework/Source/Party/Public/Chat/CommonSlashCommands.h
+++ b/Engine/Plugins/Online/OnlineFramework/Source/Party/Public/Chat/CommonSlashCommands.h
@@ -1,4 +1,4 @@
-// Copyright 1998-2018 Epic Games, Inc. All Rights Reserved.
+// Copyright 1998-2019 Epic Games, Inc. All Rights Reserved.
 
 #pragma once
 
