UE-59012 //UE4/Main - Compile UE4Game Lumin - WARNING: Precompiled module 'PurchaseFlow' will not be usable without non-precompiled module 'WebBrowser'.
-fix for this, stop precompiling the modules dependent on WebBrowser, which does not compile for Lumin.

diff --git a/Engine/Plugins/Online/OnlineFramework/Source/LoginFlow/LoginFlow.Build.cs b/Engine/Plugins/Online/OnlineFramework/Source/LoginFlow/LoginFlow.Build.cs
--- a/Engine/Plugins/Online/OnlineFramework/Source/LoginFlow/LoginFlow.Build.cs
+++ b/Engine/Plugins/Online/OnlineFramework/Source/LoginFlow/LoginFlow.Build.cs
@@ -52,5 +52,10 @@ public class LoginFlow : ModuleRules
 				"AnalyticsET",
 			}
 		);
+
+		if (Target.Platform == UnrealTargetPlatform.Lumin)
+		{
+			PrecompileForTargets = PrecompileTargetsType.None;
+		}
 	}
 }
