The memory pointed at by the FText* got destroyed, probably on the stack, and then when we switched localizations it tries to access bogus memory.

diff --git a/Engine/Plugins/Online/OnlineFramework/Source/Party/Private/Chat/ChatSlashCommands.cpp b/Engine/Plugins/Online/OnlineFramework/Source/Party/Private/Chat/ChatSlashCommands.cpp
--- a/Engine/Plugins/Online/OnlineFramework/Source/Party/Private/Chat/ChatSlashCommands.cpp
+++ b/Engine/Plugins/Online/OnlineFramework/Source/Party/Private/Chat/ChatSlashCommands.cpp
@@ -306,16 +306,16 @@ bool FRegisteredSlashCommands::SpaceWasJustTyped(const FString& NewUserText)
 void FChatSlashCommand::RecacheForLocalization() const
 {
 	//re-read the source text reference to get the new localization
-	CommandNameString = CommandNameTextSrc->ToString();
+	CommandNameString = CommandNameTextSrc.ToString();
 }
 
 //////////////////////////////////////////////////////////////////////////////////
 // Chat Slash Command
 
 FChatSlashCommand::FChatSlashCommand(const FText& InCommandName) 
-	: CommandNameTextSrc(&InCommandName)
+	: CommandNameTextSrc(InCommandName)
 {
-	//done outside of initializer list to avoid pesky headerfile order definition compile errors/warnings
+	//done outside of initializer list to avoid pesky header file order definition compile errors/warnings
 	CommandNameString = InCommandName.ToString();
 }
 
diff --git a/Engine/Plugins/Online/OnlineFramework/Source/Party/Public/Chat/ChatSlashCommands.h b/Engine/Plugins/Online/OnlineFramework/Source/Party/Public/Chat/ChatSlashCommands.h
--- a/Engine/Plugins/Online/OnlineFramework/Source/Party/Public/Chat/ChatSlashCommands.h
+++ b/Engine/Plugins/Online/OnlineFramework/Source/Party/Public/Chat/ChatSlashCommands.h
@@ -108,7 +108,7 @@ protected:
 	mutable FString CommandNameString;
 
 private:
-	const FText* CommandNameTextSrc;
+	const FText CommandNameTextSrc;
 	TWeakObjectPtr<USocialToolkit> MyToolkit;
 
 	FChatSlashCommand(const FChatSlashCommand& copy) = delete;
