Fix long ensure when encountering an invalid SocialUser

If there are no pending queries, and the UserInfo is invalid, we assume that there's nothing good to do with the User. So, remove it from tracking in the SocialToolkit. This will release it for cleanup on next GC pass.

#tests use debugger to force a UserInfo to be invalid when all queries are complete. Ensure the relevant invalid user is cleaned up by GC, and no UI elements make reference to the user.

diff --git a/Engine/Plugins/Online/OnlineFramework/Source/Party/Private/SocialToolkit.cpp b/Engine/Plugins/Online/OnlineFramework/Source/Party/Private/SocialToolkit.cpp
--- a/Engine/Plugins/Online/OnlineFramework/Source/Party/Private/SocialToolkit.cpp
+++ b/Engine/Plugins/Online/OnlineFramework/Source/Party/Private/SocialToolkit.cpp
@@ -953,4 +953,9 @@ void USocialToolkit::HandlePartyMemberExited(const FUniqueNetId& LocalUserId, co
 void USocialToolkit::HandleGameDestroyed(const FName SessionName, bool bWasSuccessful)
 {
 	// Update the recent player list whenever a game session ends
-}
\ No newline at end of file
+}
+
+void USocialToolkit::HandleUserInvalidated(USocialUser* InvalidUser)
+{
+	AllUsers.Remove(InvalidUser);
+}
diff --git a/Engine/Plugins/Online/OnlineFramework/Source/Party/Private/User/SocialUser.cpp b/Engine/Plugins/Online/OnlineFramework/Source/Party/Private/User/SocialUser.cpp
--- a/Engine/Plugins/Online/OnlineFramework/Source/Party/Private/User/SocialUser.cpp
+++ b/Engine/Plugins/Online/OnlineFramework/Source/Party/Private/User/SocialUser.cpp
@@ -318,16 +318,28 @@ void USocialUser::TryBroadcastInitializationComplete()
 	{
 		// We consider a social user to be initialized when it has valid primary OSS user info and no pending queries
 		const FSubsystemUserInfo* SubsystemInfo = SubsystemInfoByType.Find(ESocialSubsystem::Primary);
-		if (SubsystemInfo && ensureMsgf(SubsystemInfo->UserInfo.IsValid(), TEXT("SocialUser [%s] has primary subsystem info and no pending queries, but primary UserInfo is invalid!"), *ToDebugString()))
+		if (SubsystemInfo)
 		{
-			UE_LOG(LogParty, VeryVerbose, TEXT("SocialUser [%s] fully initialized."), *ToDebugString());
+			if (SubsystemInfo->UserInfo.IsValid())
+			{
+				UE_LOG(LogParty, VeryVerbose, TEXT("SocialUser [%s] fully initialized."), *ToDebugString());
 
-			bIsInitialized = true;
+				bIsInitialized = true;
 
-			FOnNewSocialUserInitialized InitEvent;
-			if (InitEventsByUser.RemoveAndCopyValue(this, InitEvent))
+				FOnNewSocialUserInitialized InitEvent;
+				if (InitEventsByUser.RemoveAndCopyValue(this, InitEvent))
+				{
+					InitEvent.Broadcast(*this);
+				}
+			}
+			else
 			{
-				InitEvent.Broadcast(*this);
+				// User is invalid with no open queries
+				// Assume that this means the sought user doesn't exist
+				InitEventsByUser.Remove(this);
+
+				// Remove Toolkit's reference to the SocialUser, GC will clean it
+				GetOwningToolkit().HandleUserInvalidated(this);
 			}
 		}
 	}
diff --git a/Engine/Plugins/Online/OnlineFramework/Source/Party/Public/SocialToolkit.h b/Engine/Plugins/Online/OnlineFramework/Source/Party/Public/SocialToolkit.h
--- a/Engine/Plugins/Online/OnlineFramework/Source/Party/Public/SocialToolkit.h
+++ b/Engine/Plugins/Online/OnlineFramework/Source/Party/Public/SocialToolkit.h
@@ -99,6 +99,8 @@ PARTY_SCOPE:
 
 	bool TrySendFriendInvite(USocialUser& SocialUser, ESocialSubsystem SubsystemType) const;
 
+	void HandleUserInvalidated(USocialUser* InvalidUser);
+
 #if PLATFORM_PS4
 	void NotifyPSNFriendsListRebuilt();
 #endif
