Disable Qos module on HTML5

diff --git a/Engine/Plugins/Online/OnlineFramework/OnlineFramework.uplugin b/Engine/Plugins/Online/OnlineFramework/OnlineFramework.uplugin
--- a/Engine/Plugins/Online/OnlineFramework/OnlineFramework.uplugin
+++ b/Engine/Plugins/Online/OnlineFramework/OnlineFramework.uplugin
@@ -13,7 +13,10 @@
 	"Modules": [
 		{
 			"Name": "Qos",
-			"Type": "Runtime"
+			"Type": "Runtime",
+			"BlacklistPlatforms": [
+				"HTML5"
+			]
 		},
 		{
 			"Name": "Party",
