Ran into a crash in SocialUser during PIE sometimes on multiple runs with mutliple users.

Dan.Hertzka

diff --git a/Engine/Plugins/Online/OnlineFramework/Source/Party/Private/User/SocialUser.cpp b/Engine/Plugins/Online/OnlineFramework/Source/Party/Private/User/SocialUser.cpp
--- a/Engine/Plugins/Online/OnlineFramework/Source/Party/Private/User/SocialUser.cpp
+++ b/Engine/Plugins/Online/OnlineFramework/Source/Party/Private/User/SocialUser.cpp
@@ -1052,19 +1052,23 @@ void USocialUser::SetSubsystemId(ESocialSubsystem SubsystemType, const FUniqueNe
 		IOnlineSubsystem* OSS = OwningToolkit.GetSocialOss(SubsystemType);
 		if (ensure(OSS))
 		{
-			TSharedPtr<FOnlineUser> UserInfo = OSS->GetUserInterface()->GetUserInfo(OwningToolkit.GetLocalUserNum(), *SubsystemId);
-			if (UserInfo.IsValid())
+			IOnlineUserPtr UserInterface = OSS->GetUserInterface();
+			if (UserInterface.IsValid())
 			{
-				SetUserInfo(SubsystemType, UserInfo.ToSharedRef());
-			}
-			else
-			{
-				UE_LOG(LogParty, VeryVerbose, TEXT("SocialUser [%s] querying user info on subsystem [%s]"), *ToDebugString(), ToString(SubsystemType));
+				TSharedPtr<FOnlineUser> UserInfo = UserInterface->GetUserInfo(OwningToolkit.GetLocalUserNum(), *SubsystemId);
+				if (UserInfo.IsValid())
+				{
+					SetUserInfo(SubsystemType, UserInfo.ToSharedRef());
+				}
+				else
+				{
+					UE_LOG(LogParty, VeryVerbose, TEXT("SocialUser [%s] querying user info on subsystem [%s]"), *ToDebugString(), ToString(SubsystemType));
 
-				// No valid user info exists on this subsystem, so queue up a query for it
-				auto QueryCompleteHandler = FSocialQuery_UserInfo::FOnQueryComplete::CreateUObject(this, &USocialUser::HandleQueryUserInfoComplete);
-				FSocialQueryManager::AddUserId<FSocialQuery_UserInfo>(OwningToolkit, SubsystemType, SubsystemId.GetUniqueNetId().ToSharedRef(), QueryCompleteHandler);
-				NumPendingQueries++;
+					// No valid user info exists on this subsystem, so queue up a query for it
+					auto QueryCompleteHandler = FSocialQuery_UserInfo::FOnQueryComplete::CreateUObject(this, &USocialUser::HandleQueryUserInfoComplete);
+					FSocialQueryManager::AddUserId<FSocialQuery_UserInfo>(OwningToolkit, SubsystemType, SubsystemId.GetUniqueNetId().ToSharedRef(), QueryCompleteHandler);
+					NumPendingQueries++;
+				}
 			}
 		}
 	}
