Correct return value logic for ApplyCrossplayRestriction.

#Fix for party join in progress failure introduced in CL 5316175
#This issue was caused by the approval action no longer being set to EnqueueAndStartBeacon or Enqueue when a request to join a party in game is made. ()

diff --git a/Engine/Plugins/Online/OnlineFramework/Source/Party/Private/Party/SocialParty.cpp b/Engine/Plugins/Online/OnlineFramework/Source/Party/Private/Party/SocialParty.cpp
--- a/Engine/Plugins/Online/OnlineFramework/Source/Party/Private/Party/SocialParty.cpp
+++ b/Engine/Plugins/Online/OnlineFramework/Source/Party/Private/Party/SocialParty.cpp
@@ -191,7 +191,7 @@ bool USocialParty::ApplyCrossplayRestriction(FPartyJoinApproval& JoinApproval, c
 		}
 	}
 
-	return JoinApproval.GetDenialReason().HasAnyReason();
+	return !JoinApproval.GetDenialReason().HasAnyReason();
 }
 
 FPartyJoinApproval USocialParty::EvaluateJoinRequest(const FUniqueNetId& PlayerId, const FUserPlatform& Platform, const FOnlinePartyData& JoinData, bool bFromJoinRequest) const
