fix cis
shadowed variables

diff --git a/Engine/Plugins/Online/OnlineFramework/Source/Party/Private/SocialDebugTools.cpp b/Engine/Plugins/Online/OnlineFramework/Source/Party/Private/SocialDebugTools.cpp
--- a/Engine/Plugins/Online/OnlineFramework/Source/Party/Private/SocialDebugTools.cpp
+++ b/Engine/Plugins/Online/OnlineFramework/Source/Party/Private/SocialDebugTools.cpp
@@ -95,13 +95,13 @@ bool USocialDebugTools::Exec(class UWorld* InWorld, const TCHAR* Cmd, FOutputDev
 
 					if (bSuccess)
 					{
-						LeaveParty(Instance, FLeavePartyComplete::CreateLambda([this, Instance, FriendName](bool bSuccess)
+						LeaveParty(Instance, FLeavePartyComplete::CreateLambda([this, Instance, FriendName](bool bLeavePartySuccess)
 						{
-							UE_LOG(LogParty, Display, TEXT("Leave party OSS context[%s] %s"), *Instance, *LexToString(bSuccess));
+							UE_LOG(LogParty, Display, TEXT("Leave party OSS context[%s] %s"), *Instance, *LexToString(bLeavePartySuccess));
 
-							JoinParty(Instance, FriendName, FJoinPartyComplete::CreateLambda([this, Instance](bool bSuccess)
+							JoinParty(Instance, FriendName, FJoinPartyComplete::CreateLambda([this, Instance](bool bJoinPartySuccess)
 							{
-								UE_LOG(LogParty, Display, TEXT("Join party OSS context[%s] %s"), *Instance, *LexToString(bSuccess));
+								UE_LOG(LogParty, Display, TEXT("Join party OSS context[%s] %s"), *Instance, *LexToString(bJoinPartySuccess));
 							}));							
 						}));
 					}
@@ -174,25 +174,25 @@ void USocialDebugTools::Login(const FString& Instance, const FOnlineAccountCrede
 				{
 					auto PresenceDelegate = FOnPresenceReceivedDelegate::CreateLambda([this, Instance, OnComplete](const class FUniqueNetId& LocalUserId, const TSharedRef<FOnlineUserPresence>& Presence)
 					{
-						FInstanceContext& Context = GetContext(Instance);
+						FInstanceContext& ContextTmp = GetContext(Instance);
 						OnComplete.ExecuteIfBound(true);
 
-						if (Context.PresenceReceivedDelegateHandle.IsValid())
+						if (ContextTmp.PresenceReceivedDelegateHandle.IsValid())
 						{
-							Context.GetOSS()->GetPresenceInterface()->ClearOnPresenceReceivedDelegate_Handle(Context.PresenceReceivedDelegateHandle);
-							Context.PresenceReceivedDelegateHandle.Reset();
+							ContextTmp.GetOSS()->GetPresenceInterface()->ClearOnPresenceReceivedDelegate_Handle(ContextTmp.PresenceReceivedDelegateHandle);
+							ContextTmp.PresenceReceivedDelegateHandle.Reset();
 						}
 					});
 
 					auto LoginDelegate = FOnLoginCompleteDelegate::CreateLambda([this, Instance, PresenceDelegate, OnComplete](int32 InLocalUserNum, bool bWasSuccessful, const FUniqueNetId& UserId, const FString& Error)
 					{
-						FInstanceContext& Context = GetContext(Instance);
+						FInstanceContext& ContextTmp = GetContext(Instance);
 						if (bWasSuccessful)
 						{
-							IOnlinePresencePtr OnlinePresence = Context.GetOSS()->GetPresenceInterface();
+							IOnlinePresencePtr OnlinePresence = ContextTmp.GetOSS()->GetPresenceInterface();
 							if (OnlinePresence.IsValid())
 							{
-								Context.PresenceReceivedDelegateHandle = OnlinePresence->AddOnPresenceReceivedDelegate_Handle(PresenceDelegate);
+								ContextTmp.PresenceReceivedDelegateHandle = OnlinePresence->AddOnPresenceReceivedDelegate_Handle(PresenceDelegate);
 
 								FOnlineUserPresenceStatus Status;
 								Status.State = EOnlinePresenceState::Online;
@@ -209,10 +209,10 @@ void USocialDebugTools::Login(const FString& Instance, const FOnlineAccountCrede
 							OnComplete.ExecuteIfBound(false);
 						}
 						
-						if (Context.LoginCompleteDelegateHandle.IsValid())
+						if (ContextTmp.LoginCompleteDelegateHandle.IsValid())
 						{
-							Context.GetOSS()->GetIdentityInterface()->ClearOnLoginCompleteDelegate_Handle(InLocalUserNum, Context.LoginCompleteDelegateHandle);
-							Context.LoginCompleteDelegateHandle.Reset();
+							ContextTmp.GetOSS()->GetIdentityInterface()->ClearOnLoginCompleteDelegate_Handle(InLocalUserNum, ContextTmp.LoginCompleteDelegateHandle);
+							ContextTmp.LoginCompleteDelegateHandle.Reset();
 						}
 					});
 
@@ -248,11 +248,11 @@ void USocialDebugTools::Logout(const FString& Instance, const FLogoutComplete& O
 					{
 						OnComplete.ExecuteIfBound(bWasSuccessful);
 
-						FInstanceContext& Context = GetContext(Instance);
-						if (Context.LogoutCompleteDelegateHandle.IsValid())
+						FInstanceContext& ContextTmp = GetContext(Instance);
+						if (ContextTmp.LogoutCompleteDelegateHandle.IsValid())
 						{
-							GetContext(Instance).GetOSS()->GetIdentityInterface()->ClearOnLogoutCompleteDelegate_Handle(InLocalUserNum, Context.LogoutCompleteDelegateHandle);
-							Context.LogoutCompleteDelegateHandle.Reset();
+							GetContext(Instance).GetOSS()->GetIdentityInterface()->ClearOnLogoutCompleteDelegate_Handle(InLocalUserNum, ContextTmp.LogoutCompleteDelegateHandle);
+							ContextTmp.LogoutCompleteDelegateHandle.Reset();
 						}
 					});
 
@@ -274,8 +274,8 @@ void USocialDebugTools::JoinParty(const FString& Instance, const FString& Friend
 		IOnlineIdentityPtr OnlineIdentity = OnlineSub->GetIdentityInterface();
 		if (OnlineIdentity.IsValid())
 		{
-			TSharedPtr<const FUniqueNetId> UserId = OnlineIdentity->GetUniquePlayerId(LocalUserNum);
-			if (UserId.IsValid())
+			TSharedPtr<const FUniqueNetId> LocalUserId = OnlineIdentity->GetUniquePlayerId(LocalUserNum);
+			if (LocalUserId.IsValid())
 			{
 				IOnlinePartyPtr OnlineParty = OnlineSub->GetPartyInterface();
 				if (OnlineParty.IsValid())
@@ -286,14 +286,14 @@ void USocialDebugTools::JoinParty(const FString& Instance, const FString& Friend
 						IOnlineUserPtr OnlineUser = OnlineSub->GetUserInterface();
 						if (OnlineUser.IsValid())
 						{
-							OnlineUser->QueryUserIdMapping(*UserId, FriendName, IOnlineUser::FOnQueryUserMappingComplete::CreateLambda([this, Instance, OnlineParty, OnComplete](bool bWasSuccessful, const FUniqueNetId& UserId, const FString& DisplayNameOrEmail, const FUniqueNetId& FoundUserId, const FString& Error)
+							OnlineUser->QueryUserIdMapping(*LocalUserId, FriendName, IOnlineUser::FOnQueryUserMappingComplete::CreateLambda([this, Instance, OnlineParty, OnComplete](bool bWasSuccessful, const FUniqueNetId& UserId, const FString& DisplayNameOrEmail, const FUniqueNetId& FoundUserId, const FString& Error)
 							{
 								if (bWasSuccessful)
 								{
 									TSharedPtr<IOnlinePartyJoinInfo> JoinInfo = OnlineParty->GetAdvertisedParty(UserId, FoundUserId, IOnlinePartySystem::GetPrimaryPartyTypeId());
 									if (JoinInfo.IsValid())
 									{
-										OnlineParty->JoinParty(UserId, *JoinInfo, FOnJoinPartyComplete::CreateLambda([this, Instance, OnlineParty, OnComplete](const FUniqueNetId& LocalUserId, const FOnlinePartyId& PartyId, const EJoinPartyCompletionResult Result, const int32 NotApprovedReason)
+										OnlineParty->JoinParty(UserId, *JoinInfo, FOnJoinPartyComplete::CreateLambda([this, Instance, OnlineParty, OnComplete](const FUniqueNetId& UserIdTmp, const FOnlinePartyId& PartyId, const EJoinPartyCompletionResult Result, const int32 NotApprovedReason)
 										{
 											bool bSuccess = Result == EJoinPartyCompletionResult::Succeeded;
 											if (bSuccess)
@@ -301,7 +301,7 @@ void USocialDebugTools::JoinParty(const FString& Instance, const FString& Friend
 												TSharedPtr<FOnlinePartyData> PartyMemberData = GetContext(Instance).GetPartyMemberData();
 												if (PartyMemberData.IsValid())
 												{
-													OnlineParty->UpdatePartyMemberData(LocalUserId, PartyId, *PartyMemberData);
+													OnlineParty->UpdatePartyMemberData(UserIdTmp, PartyId, *PartyMemberData);
 												}
 											}
 											else
@@ -323,7 +323,7 @@ void USocialDebugTools::JoinParty(const FString& Instance, const FString& Friend
 						TSharedPtr<IOnlinePartyJoinInfo> JoinInfo = GetDefaultPartyJoinInfo();
 						if (JoinInfo.IsValid())
 						{
-							OnlineParty->JoinParty(*UserId, *JoinInfo, FOnJoinPartyComplete::CreateLambda([this, Instance, OnlineParty, OnComplete](const FUniqueNetId& LocalUserId, const FOnlinePartyId& PartyId, const EJoinPartyCompletionResult Result, const int32 NotApprovedReason)
+							OnlineParty->JoinParty(*LocalUserId, *JoinInfo, FOnJoinPartyComplete::CreateLambda([this, Instance, OnlineParty, OnComplete](const FUniqueNetId& UserId, const FOnlinePartyId& PartyId, const EJoinPartyCompletionResult Result, const int32 NotApprovedReason)
 							{
 								bool bSuccess = Result == EJoinPartyCompletionResult::Succeeded;
 								if (bSuccess)
@@ -331,7 +331,7 @@ void USocialDebugTools::JoinParty(const FString& Instance, const FString& Friend
 									TSharedPtr<FOnlinePartyData> PartyMemberData = GetContext(Instance).GetPartyMemberData();
 									if (PartyMemberData.IsValid())
 									{
-										OnlineParty->UpdatePartyMemberData(LocalUserId, PartyId, *PartyMemberData);
+										OnlineParty->UpdatePartyMemberData(UserId, PartyId, *PartyMemberData);
 									}
 								}
 								else
@@ -629,7 +629,7 @@ void USocialDebugTools::HandlePartyInviteReceived(const FUniqueNetId& LocalUserI
 					if (*Invite->GetPartyId() == PartyId)
 					{						
 						const FString Instance = Context->Name;
-						OnlineParty->JoinParty(LocalUserId, *Invite, FOnJoinPartyComplete::CreateLambda([this, Instance, OnlineParty](const FUniqueNetId& LocalUserId, const FOnlinePartyId& PartyId, const EJoinPartyCompletionResult Result, const int32 NotApprovedReason)
+						OnlineParty->JoinParty(LocalUserId, *Invite, FOnJoinPartyComplete::CreateLambda([this, Instance, OnlineParty](const FUniqueNetId& UserId, const FOnlinePartyId& PartyIdTmp, const EJoinPartyCompletionResult Result, const int32 NotApprovedReason)
 						{
 							bool bSuccess = Result == EJoinPartyCompletionResult::Succeeded;
 							if (bSuccess)
@@ -637,7 +637,7 @@ void USocialDebugTools::HandlePartyInviteReceived(const FUniqueNetId& LocalUserI
 								TSharedPtr<FOnlinePartyData> PartyMemberData = GetContext(Instance).GetPartyMemberData();
 								if (PartyMemberData.IsValid())
 								{
-									OnlineParty->UpdatePartyMemberData(LocalUserId, PartyId, *PartyMemberData);
+									OnlineParty->UpdatePartyMemberData(UserId, PartyIdTmp, *PartyMemberData);
 								}
 							}
 							else
