cis fix for member init order

diff --git a/Engine/Plugins/Online/OnlineFramework/Source/Party/Public/SocialDebugTools.h b/Engine/Plugins/Online/OnlineFramework/Source/Party/Public/SocialDebugTools.h
--- a/Engine/Plugins/Online/OnlineFramework/Source/Party/Public/SocialDebugTools.h
+++ b/Engine/Plugins/Online/OnlineFramework/Source/Party/Public/SocialDebugTools.h
@@ -51,8 +51,8 @@ private:
 	struct FInstanceContext
 	{
 		FInstanceContext(const FString& InstanceName, USocialDebugTools& SocialDebugTools)
-			: OnlineSub(nullptr)
-			, Name(InstanceName)
+			: Name(InstanceName)
+			, OnlineSub(nullptr)
 			, Owner(SocialDebugTools)
 		{}
 
