Use IOnlinePartySystem::ClearInvitations to remove the invitation from OSS instead of AcceptInvitation.

diff --git a/Engine/Plugins/Online/OnlineFramework/Source/Party/Private/User/SocialUser.cpp b/Engine/Plugins/Online/OnlineFramework/Source/Party/Private/User/SocialUser.cpp
--- a/Engine/Plugins/Online/OnlineFramework/Source/Party/Private/User/SocialUser.cpp
+++ b/Engine/Plugins/Online/OnlineFramework/Source/Party/Private/User/SocialUser.cpp
@@ -685,7 +685,7 @@ void USocialUser::JoinParty(const FOnlinePartyTypeId& PartyTypeId) const
 	if (bHasSentInvite)
 	{
 		IOnlinePartyPtr PartyInterface = Online::GetPartyInterfaceChecked(GetWorld());
-		PartyInterface->AcceptInvitation(*GetOwningToolkit().GetLocalUserNetId(ESocialSubsystem::Primary), *GetUserId(ESocialSubsystem::Primary));
+		PartyInterface->ClearInvitations(*GetOwningToolkit().GetLocalUserNetId(ESocialSubsystem::Primary), *GetUserId(ESocialSubsystem::Primary), nullptr);
 		OnPartyInviteAccepted().Broadcast();
 	}
 }
