Add ToPlatformIconKey() function to FUserPlatform

diff --git a/Engine/Plugins/Online/OnlineFramework/Source/Party/Private/SocialTypes.cpp b/Engine/Plugins/Online/OnlineFramework/Source/Party/Private/SocialTypes.cpp
--- a/Engine/Plugins/Online/OnlineFramework/Source/Party/Private/SocialTypes.cpp
+++ b/Engine/Plugins/Online/OnlineFramework/Source/Party/Private/SocialTypes.cpp
@@ -7,6 +7,10 @@
 // FUserPlatform
 //////////////////////////////////////////////////////////////////////////
 
+#define PLATFORM_NAME_PC		TEXT("PC")
+#define PLATFORM_NAME_CONSOLE	TEXT("CONSOLE")
+#define PLATFORM_NAME_MOBILE	TEXT("MOBILE")
+
 bool FUserPlatform::operator==(const FString& OtherStr) const
 {
 	return PlatformStr == OtherStr;
@@ -17,6 +21,26 @@ bool FUserPlatform::operator==(const FUserPlatform& Other) const
 	return PlatformStr == Other.PlatformStr;
 }
 
+const FString FUserPlatform::ToPlatformIconKey() const
+{
+	if (IsDesktop())
+	{
+		return PLATFORM_NAME_PC;
+	}
+
+	if (IsMobile())
+	{
+		return PLATFORM_NAME_MOBILE;
+	}
+
+	FUserPlatform LocalPlatform = FUserPlatform(IOnlineSubsystem::GetLocalPlatformName());
+	if (IsConsole() && LocalPlatform.IsConsole() && PlatformStr != LocalPlatform)
+	{
+		return PLATFORM_NAME_CONSOLE;
+	}
+	return PlatformStr;
+}
+
 bool FUserPlatform::IsValid() const
 {
 	return !PlatformStr.IsEmpty();
diff --git a/Engine/Plugins/Online/OnlineFramework/Source/Party/Public/SocialTypes.h b/Engine/Plugins/Online/OnlineFramework/Source/Party/Public/SocialTypes.h
--- a/Engine/Plugins/Online/OnlineFramework/Source/Party/Public/SocialTypes.h
+++ b/Engine/Plugins/Online/OnlineFramework/Source/Party/Public/SocialTypes.h
@@ -91,6 +91,7 @@ public:
 
 	operator const FString&() const { return PlatformStr; }
 	const FString& ToString() const { return PlatformStr; }
+	const FString ToPlatformIconKey() const;
 
 	bool operator==(const FString& OtherStr) const;
 	bool operator!=(const FString& OtherStr) const { return !operator==(OtherStr); }
