Change the function name in FUserPlatform from ToPlatformIconKey to GetTypeName.

diff --git a/Engine/Plugins/Online/OnlineFramework/Source/Party/Private/SocialTypes.cpp b/Engine/Plugins/Online/OnlineFramework/Source/Party/Private/SocialTypes.cpp
--- a/Engine/Plugins/Online/OnlineFramework/Source/Party/Private/SocialTypes.cpp
+++ b/Engine/Plugins/Online/OnlineFramework/Source/Party/Private/SocialTypes.cpp
@@ -21,7 +21,7 @@ bool FUserPlatform::operator==(const FUserPlatform& Other) const
 	return PlatformStr == Other.PlatformStr;
 }
 
-const FString FUserPlatform::ToPlatformIconKey() const
+const FString FUserPlatform::GetTypeName() const
 {
 	if (IsDesktop())
 	{
diff --git a/Engine/Plugins/Online/OnlineFramework/Source/Party/Public/SocialTypes.h b/Engine/Plugins/Online/OnlineFramework/Source/Party/Public/SocialTypes.h
--- a/Engine/Plugins/Online/OnlineFramework/Source/Party/Public/SocialTypes.h
+++ b/Engine/Plugins/Online/OnlineFramework/Source/Party/Public/SocialTypes.h
@@ -91,7 +91,7 @@ public:
 
 	operator const FString&() const { return PlatformStr; }
 	const FString& ToString() const { return PlatformStr; }
-	const FString ToPlatformIconKey() const;
+	const FString GetTypeName() const;
 
 	bool operator==(const FString& OtherStr) const;
 	bool operator!=(const FString& OtherStr) const { return !operator==(OtherStr); }
