Added code to destroy the FInstanceContext's OSS on PIE shutdown.

#tests Compile DebugGame Editor Win64 / Shipping Client PS4 / Verified the same golem works properly when running PIE multiple times in a row

diff --git a/Engine/Plugins/Online/OnlineFramework/Source/Party/Private/SocialDebugTools.cpp b/Engine/Plugins/Online/OnlineFramework/Source/Party/Private/SocialDebugTools.cpp
--- a/Engine/Plugins/Online/OnlineFramework/Source/Party/Private/SocialDebugTools.cpp
+++ b/Engine/Plugins/Online/OnlineFramework/Source/Party/Private/SocialDebugTools.cpp
@@ -593,6 +593,9 @@ void USocialDebugTools::FInstanceContext::Shutdown()
 			OnlineParty.Reset();
 		}
 		OnlineSub->Shutdown();
+
+		const FString OSSName = FName(MCP_SUBSYSTEM).ToString() + TEXT(":") + Name;
+		IOnlineSubsystem::Destroy(*OSSName);
 	}
 }
 
