Fix ensure firing when modifying controller ids in debug instances. This is now an error log as it is not show stopping.

diff --git a/Engine/Plugins/Online/OnlineFramework/Source/Party/Private/SocialToolkit.cpp b/Engine/Plugins/Online/OnlineFramework/Source/Party/Private/SocialToolkit.cpp
--- a/Engine/Plugins/Online/OnlineFramework/Source/Party/Private/SocialToolkit.cpp
+++ b/Engine/Plugins/Online/OnlineFramework/Source/Party/Private/SocialToolkit.cpp
@@ -607,8 +607,10 @@ void USocialToolkit::HandlePlayerLoginStatusChanged(int32 LocalUserNum, ELoginSt
 	{
 		if (NewStatus == ELoginStatus::LoggedIn)
 		{
-			if (!ensure(AllUsers.Num() == 0))
+			if (AllUsers.Num() != 0)
 			{
+				UE_LOG(LogParty, Error, TEXT("HandlePlayerLoginStatusChanged: Changed login status but we were not informed their status had changed previously"));
+
 				// Nobody told us we logged out! Handle it now just so we're fresh, but not good!
 				OnOwnerLoggedOut();
 			}
