Fix invalid player lookup that can lead to a nullptr dereference when working with indexes

diff --git a/Engine/Plugins/Online/OnlineFramework/Source/Party/Private/SocialManager.cpp b/Engine/Plugins/Online/OnlineFramework/Source/Party/Private/SocialManager.cpp
--- a/Engine/Plugins/Online/OnlineFramework/Source/Party/Private/SocialManager.cpp
+++ b/Engine/Plugins/Online/OnlineFramework/Source/Party/Private/SocialManager.cpp
@@ -826,7 +826,7 @@ void USocialManager::HandleWorldEstablished(UWorld* World)
 
 void USocialManager::HandleLocalPlayerAdded(int32 LocalUserNum)
 {
-	ULocalPlayer* NewLocalPlayer = GetGameInstance().FindLocalPlayerFromControllerId(LocalUserNum);
+	ULocalPlayer* NewLocalPlayer = GetGameInstance().GetLocalPlayerByIndex(LocalUserNum);
 	check(NewLocalPlayer);
 
 	CreateSocialToolkit(*NewLocalPlayer);
diff --git a/Engine/Source/Runtime/Engine/Classes/Engine/GameInstance.h b/Engine/Source/Runtime/Engine/Classes/Engine/GameInstance.h
--- a/Engine/Source/Runtime/Engine/Classes/Engine/GameInstance.h
+++ b/Engine/Source/Runtime/Engine/Classes/Engine/GameInstance.h
@@ -260,7 +260,7 @@ public:
 	 * @param	NewPlayer	the player to add
 	 * @param	ControllerId id of the controller associated with the player
 	 */
-	int32					AddLocalPlayer(ULocalPlayer * NewPlayer, int32 ControllerId);
+	virtual int32			AddLocalPlayer(ULocalPlayer* NewPlayer, int32 ControllerId);
 
 
 	/**
