QueueUserDependentAction now ensures when a call to it attempts to create a new SocialUser if the toolkit's owner isn't actually logged in yet

diff --git a/Engine/Plugins/Online/OnlineFramework/Source/Party/Private/SocialToolkit.cpp b/Engine/Plugins/Online/OnlineFramework/Source/Party/Private/SocialToolkit.cpp
--- a/Engine/Plugins/Online/OnlineFramework/Source/Party/Private/SocialToolkit.cpp
+++ b/Engine/Plugins/Online/OnlineFramework/Source/Party/Private/SocialToolkit.cpp
@@ -307,7 +307,7 @@ void USocialToolkit::QueueUserDependentActionInternal(const FUniqueNetIdRepl& Su
 	}
 	
 	USocialUser* User = FindUser(SubsystemId);
-	if (!User)
+	if (!User && ensureMsgf(IsOwnerLoggedIn(), TEXT("Cannot QueueUserDependentAction while local user is logged out! Toolkit [%d], ID [%s], Subsystem [%s]"), GetLocalUserNum(), *SubsystemId.ToDebugString(), *ToString(SubsystemType)))
 	{
 		if (SubsystemType == ESocialSubsystem::Primary)
 		{
