Reverting change in 6638190 as the RestoreParties flow can somehow result in creation of a SocialUser (via PartyMember) before the local user is formally "logged in".

diff --git a/Engine/Plugins/Online/OnlineFramework/Source/Party/Private/SocialToolkit.cpp b/Engine/Plugins/Online/OnlineFramework/Source/Party/Private/SocialToolkit.cpp
--- a/Engine/Plugins/Online/OnlineFramework/Source/Party/Private/SocialToolkit.cpp
+++ b/Engine/Plugins/Online/OnlineFramework/Source/Party/Private/SocialToolkit.cpp
@@ -307,7 +307,8 @@ void USocialToolkit::QueueUserDependentActionInternal(const FUniqueNetIdRepl& Su
 	}
 	
 	USocialUser* User = FindUser(SubsystemId);
-	if (!User && ensureMsgf(IsOwnerLoggedIn(), TEXT("Cannot QueueUserDependentAction while local user is logged out! Toolkit [%d], ID [%s], Subsystem [%s]"), GetLocalUserNum(), *SubsystemId.ToDebugString(), *ToString(SubsystemType)))
+
+	if (!User)// && ensureMsgf(IsOwnerLoggedIn(), TEXT("Cannot QueueUserDependentAction while local user is logged out! Toolkit [%d], ID [%s], Subsystem [%s]"), GetLocalUserNum(), *SubsystemId.ToDebugString(), *ToString(SubsystemType)))
 	{
 		if (SubsystemType == ESocialSubsystem::Primary)
 		{
