Reimplemented CL 6638190 to better match the original intent.

Game overloads of IsOwnerLoggedIn() (especially where they add stricter game-specific requirements) aren't intended to affect this usage.

diff --git a/Engine/Plugins/Online/OnlineFramework/Source/Party/Private/SocialToolkit.cpp b/Engine/Plugins/Online/OnlineFramework/Source/Party/Private/SocialToolkit.cpp
--- a/Engine/Plugins/Online/OnlineFramework/Source/Party/Private/SocialToolkit.cpp
+++ b/Engine/Plugins/Online/OnlineFramework/Source/Party/Private/SocialToolkit.cpp
@@ -308,7 +308,7 @@ void USocialToolkit::QueueUserDependentActionInternal(const FUniqueNetIdRepl& Su
 	
 	USocialUser* User = FindUser(SubsystemId);
 
-	if (!User)// && ensureMsgf(IsOwnerLoggedIn(), TEXT("Cannot QueueUserDependentAction while local user is logged out! Toolkit [%d], ID [%s], Subsystem [%s]"), GetLocalUserNum(), *SubsystemId.ToDebugString(), *ToString(SubsystemType)))
+	if (!User && ensureMsgf(USocialToolkit::IsOwnerLoggedIn(), TEXT("Cannot QueueUserDependentAction while local user is logged out! Toolkit [%d], ID [%s], Subsystem [%s]"), GetLocalUserNum(), *SubsystemId.ToDebugString(), ToString(SubsystemType)))
 	{
 		if (SubsystemType == ESocialSubsystem::Primary)
 		{
