Skip asset hotfixing while in kairos
Always allow OSS patch check to complete even if there is no signed in user beyond the first attempt

diff --git a/Engine/Plugins/Online/OnlineFramework/Source/Hotfix/Private/UpdateManager.cpp b/Engine/Plugins/Online/OnlineFramework/Source/Hotfix/Private/UpdateManager.cpp
--- a/Engine/Plugins/Online/OnlineFramework/Source/Hotfix/Private/UpdateManager.cpp
+++ b/Engine/Plugins/Online/OnlineFramework/Source/Hotfix/Private/UpdateManager.cpp
@@ -317,7 +317,7 @@ void UUpdateManager::StartOSSPatchCheck()
 			{
 				const int32 ControllerId = LP->GetControllerId();
 				TSharedPtr<const FUniqueNetId> UserId = IdentityInt->GetUniquePlayerId(ControllerId);
-				if (!bInitialUpdateFinished && !UserId.IsValid())
+				if (!UserId.IsValid())
 				{
 					// Invalid user for "before title/login" check, underlying code doesn't need a valid user currently
 					UserId = IdentityInt->CreateUniquePlayerId(TEXT("InvalidUser"));
@@ -329,11 +329,6 @@ void UUpdateManager::StartOSSPatchCheck()
 					IdentityInt->GetUserPrivilege(*UserId,
 						EUserPrivileges::CanPlayOnline, IOnlineIdentity::FOnGetUserPrivilegeCompleteDelegate::CreateUObject(this, &ThisClass::OnCheckForPatchComplete, false));
 				}
-				else
-				{
-					UE_LOG(LogHotfixManager, Warning, TEXT("No valid user id when starting patch check!"));
-					PatchResult = EPatchCheckResult::NoLoggedInUser;
-				}
 			}
 			else
 			{
diff --git a/Engine/Plugins/Online/OnlineFramework/Source/Hotfix/Public/OnlineHotfixManager.h b/Engine/Plugins/Online/OnlineFramework/Source/Hotfix/Public/OnlineHotfixManager.h
--- a/Engine/Plugins/Online/OnlineFramework/Source/Hotfix/Public/OnlineHotfixManager.h
+++ b/Engine/Plugins/Online/OnlineFramework/Source/Hotfix/Public/OnlineHotfixManager.h
@@ -286,7 +286,7 @@ protected:
 	void UpdateProgress(uint32 FileCount, uint64 UpdateSize);
 
 	/** Called after any hotfixes are applied to apply last-second changes to certain asset types from .ini file data */
-	void PatchAssetsFromIniFiles();
+	virtual void PatchAssetsFromIniFiles();
 	
 	/** Used in PatchAssetsFromIniFiles to hotfix only a row in a table. */
 	void HotfixRowUpdate(UObject* Asset, const FString& AssetPath, const FString& RowName, const FString& ColumnName, const FString& NewValue, TArray<FString>& ProblemStrings);
