Adding Endpoint Fields to QoS payload

diff --git a/Engine/Plugins/Online/OnlineFramework/Source/Qos/Private/QosStats.cpp b/Engine/Plugins/Online/OnlineFramework/Source/Qos/Private/QosStats.cpp
--- a/Engine/Plugins/Online/OnlineFramework/Source/Qos/Private/QosStats.cpp
+++ b/Engine/Plugins/Online/OnlineFramework/Source/Qos/Private/QosStats.cpp
@@ -29,6 +29,8 @@ const FString FQosDatacenterStats::QosStats_NumSuccessCount = TEXT("NumSuccessCo
 const FString FQosDatacenterStats::QosStats_NetworkType = TEXT("NetworkType");
 const FString FQosDatacenterStats::QosStats_BestRegionId = TEXT("BestRegionId");
 const FString FQosDatacenterStats::QosStats_BestRegionPing = TEXT("BestRegionPing");
+const FString FQosDatacenterStats::QosStats_BestEndpointId = TEXT("BestEndpointId");
+const FString FQosDatacenterStats::QosStats_BestEndpointPing = TEXT("BestEndpointPing");
 
 /**
  * Debug output for the contents of a recorded stats event
@@ -153,6 +155,8 @@ void FQosDatacenterStats::Upload(TSharedPtr<IAnalyticsProvider>& AnalyticsProvid
  * @EventParam NetworkType string type of network the client is connected to. (Unknown, None, AirplaneMode, Cell, Wifi, Ethernet) are possible values. Will be Unknown on PC and Switch.
  * @EventParam BestRegionId string RegionId with best ping (that is usable, UNREACHABLE if none pass QoS)
  * @EventParam BestRegionPing integer ping in the best RegionId (that is usable, 0 if BestRegionId is UNREACHABLE)
+ * @EventParam BestEndpointId string RegionId with best ping (regardless of usability, UNREACHABLE if none pass QoS)
+ * @EventParam BestEndpointPing integer ping in the best RegionId (regardless of usability, 0 if BestEndpointId is UNREACHABLE)
  * @EventParam RegionDetails json representation of ping details
  * @Comments Analytics data for a complete qos datacenter determination attempt
  * 
@@ -175,12 +179,21 @@ void FQosDatacenterStats::ParseQosResults(TSharedPtr<IAnalyticsProvider>& Analyt
 	{
 		FString BestRegionId(TEXT("Unknown"));
 		int32 BestPing = INT_MAX;
+		FString BestEndpointId(TEXT("Unknown"));
+		int32 BestEndpointPing = INT_MAX;
+
 		for (const FQosStats_RegionInfo& Region : QosData.Regions)
 		{
-			if (Region.AvgPing < BestPing && Region.bUsable)
+			if (Region.AvgPing < BestPing)
 			{
-				BestRegionId = Region.RegionId;
-				BestPing = Region.AvgPing;
+				if (Region.bUsable)
+				{
+					BestRegionId = Region.RegionId;
+					BestPing = Region.AvgPing;
+				}
+				
+				BestEndpointId = Region.RegionId;
+				BestEndpointPing = Region.AvgPing;
 			}
 		}
 
@@ -195,6 +208,18 @@ void FQosDatacenterStats::ParseQosResults(TSharedPtr<IAnalyticsProvider>& Analyt
 			QoSAttributes.Add(FAnalyticsEventAttribute(QosStats_BestRegionId, TEXT("UNREACHABLE")));
 			QoSAttributes.Add(FAnalyticsEventAttribute(QosStats_BestRegionPing, 0));
 		}
+
+		BestEndpointPing = FMath::Clamp(BestEndpointPing, 0, UNREACHABLE_PING);
+		if (BestEndpointPing < UNREACHABLE_PING)
+		{
+			QoSAttributes.Add(FAnalyticsEventAttribute(QosStats_BestEndpointId, BestEndpointId));
+			QoSAttributes.Add(FAnalyticsEventAttribute(QosStats_BestEndpointPing, BestEndpointPing));
+		}
+		else
+		{
+			QoSAttributes.Add(FAnalyticsEventAttribute(QosStats_BestEndpointId, TEXT("UNREACHABLE")));
+			QoSAttributes.Add(FAnalyticsEventAttribute(QosStats_BestEndpointPing, 0));
+		}
 	}
 
 	{
diff --git a/Engine/Plugins/Online/OnlineFramework/Source/Qos/Private/QosStats.h b/Engine/Plugins/Online/OnlineFramework/Source/Qos/Private/QosStats.h
--- a/Engine/Plugins/Online/OnlineFramework/Source/Qos/Private/QosStats.h
+++ b/Engine/Plugins/Online/OnlineFramework/Source/Qos/Private/QosStats.h
@@ -170,6 +170,8 @@ private:
 	static const FString QosStats_NetworkType;
 	static const FString QosStats_BestRegionId;
 	static const FString QosStats_BestRegionPing;
+	static const FString QosStats_BestEndpointId;
+	static const FString QosStats_BestEndpointPing;
 
 	/** Version of the stats for separation */
 	int32 StatsVersion;
