Profile scopes to around presence updates to identify the cause of hitches

diff --git a/Engine/Plugins/Online/OnlineFramework/Source/Party/Private/SocialToolkit.cpp b/Engine/Plugins/Online/OnlineFramework/Source/Party/Private/SocialToolkit.cpp
--- a/Engine/Plugins/Online/OnlineFramework/Source/Party/Private/SocialToolkit.cpp
+++ b/Engine/Plugins/Online/OnlineFramework/Source/Party/Private/SocialToolkit.cpp
@@ -736,12 +736,14 @@ void USocialToolkit::HandleMapExternalIdComplete(ESocialSubsystem SubsystemType,
 
 void USocialToolkit::HandlePresenceReceived(const FUniqueNetId& UserId, const TSharedRef<FOnlineUserPresence>& NewPresence, ESocialSubsystem SubsystemType)
 {
+	QUICK_SCOPE_CYCLE_COUNTER(STAT_USocialToolkit_HandlePresenceReceived);
 	if (USocialUser* UpdatedUser = FindUser(UserId))
 	{
 		UpdatedUser->NotifyPresenceChanged(SubsystemType);
 	}
 	else if (SubsystemType == ESocialSubsystem::Platform)
 	{
+		QUICK_SCOPE_CYCLE_COUNTER(STAT_USocialToolkit_HandlePresenceReceived_Error);
 		FString ErrorString = TEXT("Platform presence received, but existing SocialUser could not be found.\n");
 		ErrorString += TEXT("Incoming UserId is ") + UserId.ToString() + TEXT(", as a UniqueIdRepl it's ") + FUniqueNetIdRepl(UserId).ToString();
 
