Additional logging for a crash in FOnlinePartySystemMcpV2::JoinParty

diff --git a/Engine/Plugins/Online/OnlineFramework/Source/Party/Private/SocialManager.cpp b/Engine/Plugins/Online/OnlineFramework/Source/Party/Private/SocialManager.cpp
--- a/Engine/Plugins/Online/OnlineFramework/Source/Party/Private/SocialManager.cpp
+++ b/Engine/Plugins/Online/OnlineFramework/Source/Party/Private/SocialManager.cpp
@@ -668,7 +668,7 @@ void USocialManager::JoinPartyInternal(FJoinPartyAttempt& JoinAttempt)
 {
 	IOnlinePartyPtr PartyInterface = Online::GetPartyInterfaceChecked(GetWorld());
 	FUniqueNetIdRepl LocalUserId = GetFirstLocalUserId(ESocialSubsystem::Primary);
-	check(LocalUserId.IsValid());
+	ensureMsgf(LocalUserId.IsValid(), TEXT("USocialManager::JoinPartyInternal: Invalid LocalUserId!"));
 
 	JoinAttempt.ActionTimeTracker.BeginStep(FJoinPartyAttempt::Step_JoinParty);
 
