Deleting unused code around StayWithPartyOnExit

diff --git a/Engine/Plugins/Online/OnlineFramework/Source/Party/Private/Party/SocialParty.cpp b/Engine/Plugins/Online/OnlineFramework/Source/Party/Private/Party/SocialParty.cpp
--- a/Engine/Plugins/Online/OnlineFramework/Source/Party/Private/Party/SocialParty.cpp
+++ b/Engine/Plugins/Online/OnlineFramework/Source/Party/Private/Party/SocialParty.cpp
@@ -1056,16 +1056,6 @@ bool USocialParty::IsCurrentlyCrossplaying() const
 	return false;
 }
 
-void USocialParty::StayWithPartyOnExit(bool bInStayWithParty)
-{
-	bStayWithPartyOnDisconnect = bInStayWithParty;
-}
-
-bool USocialParty::ShouldStayWithPartyOnExit() const
-{
-	return bStayWithPartyOnDisconnect;
-}
-
 bool USocialParty::IsPartyFunctionalityDegraded() const
 {
 	return bIsMissingXmppConnection.Get(false) || bIsMissingPlatformSession || bIsRequestingShutdown.Get(false);
diff --git a/Engine/Plugins/Online/OnlineFramework/Source/Party/Public/Party/SocialParty.h b/Engine/Plugins/Online/OnlineFramework/Source/Party/Public/Party/SocialParty.h
--- a/Engine/Plugins/Online/OnlineFramework/Source/Party/Public/Party/SocialParty.h
+++ b/Engine/Plugins/Online/OnlineFramework/Source/Party/Public/Party/SocialParty.h
@@ -131,9 +131,6 @@ public:
 	EPartyState GetOssPartyPreviousState() const;
 
 	bool IsCurrentlyCrossplaying() const;
-	void StayWithPartyOnExit(bool bInStayWithParty);
-	bool ShouldStayWithPartyOnExit() const;
-
 	bool IsPartyFunctionalityDegraded() const;
 
 	bool IsPartyFull() const;
