SocialManager reset all pending join attempts when local user logged out, this'll prevent crash/assert when leaving a party to join another while logging out.

diff --git a/Engine/Plugins/Online/OnlineFramework/Source/Party/Private/SocialManager.cpp b/Engine/Plugins/Online/OnlineFramework/Source/Party/Private/SocialManager.cpp
--- a/Engine/Plugins/Online/OnlineFramework/Source/Party/Private/SocialManager.cpp
+++ b/Engine/Plugins/Online/OnlineFramework/Source/Party/Private/SocialManager.cpp
@@ -645,6 +645,7 @@ USocialToolkit& USocialManager::CreateSocialToolkit(ULocalPlayer& OwningLocalPla
 	SocialToolkits.Add(NewToolkit);
 	NewToolkit->InitializeToolkit(OwningLocalPlayer);
 	OnToolkitCreatedInternal(*NewToolkit);
+	NewToolkit->OnToolkitReset().AddUObject(this, &USocialManager::HandleToolkitReset, NewToolkit->GetLocalUserNum());
 	return *NewToolkit;
 }
 
@@ -668,7 +669,7 @@ void USocialManager::JoinPartyInternal(FJoinPartyAttempt& JoinAttempt)
 {
 	IOnlinePartyPtr PartyInterface = Online::GetPartyInterfaceChecked(GetWorld());
 	FUniqueNetIdRepl LocalUserId = GetFirstLocalUserId(ESocialSubsystem::Primary);
-	ensureMsgf(LocalUserId.IsValid(), TEXT("USocialManager::JoinPartyInternal: Invalid LocalUserId!"));
+	checkf(LocalUserId.IsValid(), TEXT("USocialManager::JoinPartyInternal: Invalid LocalUserId!"));
 
 	JoinAttempt.ActionTimeTracker.BeginStep(FJoinPartyAttempt::Step_JoinParty);
 
@@ -677,7 +678,7 @@ void USocialManager::JoinPartyInternal(FJoinPartyAttempt& JoinAttempt)
 		UE_LOG(LogParty, Verbose, TEXT("Attempting to rejoin party [%s] now."), *JoinAttempt.RejoinInfo->PartyId->ToDebugString());
 
 		// Rejoin attempts are initiated differently, but the handler/follow-up is identical to a normal join
-		PartyInterface->RejoinParty(*GetFirstLocalUserId(ESocialSubsystem::Primary), *JoinAttempt.RejoinInfo->PartyId, IOnlinePartySystem::GetPrimaryPartyTypeId(), JoinAttempt.RejoinInfo->MemberIds, FOnJoinPartyComplete::CreateUObject(this, &USocialManager::HandleJoinPartyComplete, IOnlinePartySystem::GetPrimaryPartyTypeId()));
+		PartyInterface->RejoinParty(*LocalUserId, *JoinAttempt.RejoinInfo->PartyId, IOnlinePartySystem::GetPrimaryPartyTypeId(), JoinAttempt.RejoinInfo->MemberIds, FOnJoinPartyComplete::CreateUObject(this, &USocialManager::HandleJoinPartyComplete, IOnlinePartySystem::GetPrimaryPartyTypeId()));
 	}
 	else
 	{
@@ -851,6 +852,11 @@ void USocialManager::HandleLocalPlayerRemoved(int32 LocalUserNum)
 	}
 }
 
+void USocialManager::HandleToolkitReset(int32 LocalUserNum)
+{
+	JoinAttemptsByTypeId.Reset();
+}
+
 void USocialManager::RestorePartyStateFromPartySystem(const FOnRestorePartyStateFromPartySystemComplete& OnRestoreComplete)
 {
 	UE_LOG(LogParty, Verbose, TEXT("RestorePartyStateFromPartySystem"));
diff --git a/Engine/Plugins/Online/OnlineFramework/Source/Party/Public/SocialManager.h b/Engine/Plugins/Online/OnlineFramework/Source/Party/Public/SocialManager.h
--- a/Engine/Plugins/Online/OnlineFramework/Source/Party/Public/SocialManager.h
+++ b/Engine/Plugins/Online/OnlineFramework/Source/Party/Public/SocialManager.h
@@ -209,6 +209,7 @@ private:	// Handlers
 	void HandleWorldEstablished(UWorld* World);
 	void HandleLocalPlayerAdded(int32 LocalUserNum);
 	void HandleLocalPlayerRemoved(int32 LocalUserNum);
+	void HandleToolkitReset(int32 LocalUserNum);
 	
 	void OnRestorePartiesComplete(const FUniqueNetId& LocalUserId, const FOnlineError& Result, const FOnRestorePartyStateFromPartySystemComplete OnRestoreComplete);
 	void HandleQueryJoinabilityComplete(const FUniqueNetId& LocalUserId, const FOnlinePartyId& PartyId, EJoinPartyCompletionResult Result, int32 NotApprovedReasonCode, FOnlinePartyTypeId PartyTypeId);
