Allow party join failure subcodes for reasons other than NotApproved

diff --git a/Engine/Plugins/Online/OnlineFramework/Source/Party/Private/Party/PartyTypes.cpp b/Engine/Plugins/Online/OnlineFramework/Source/Party/Private/Party/PartyTypes.cpp
--- a/Engine/Plugins/Online/OnlineFramework/Source/Party/Private/Party/PartyTypes.cpp
+++ b/Engine/Plugins/Online/OnlineFramework/Source/Party/Private/Party/PartyTypes.cpp
@@ -82,6 +82,19 @@ FJoinPartyResult::FJoinPartyResult(EJoinPartyCompletionResult InResult, FPartyJo
 	}
 }
 
+FJoinPartyResult::FJoinPartyResult(EJoinPartyCompletionResult InResult, int32 InResultSubCode)
+{
+	SetResult(InResult);
+	if (InResult == EJoinPartyCompletionResult::NotApproved)
+	{
+		SetDenialReason(InResultSubCode);
+	}
+	else
+	{
+		ResultSubCode = InResultSubCode;
+	}
+}
+
 void FJoinPartyResult::SetDenialReason(FPartyJoinDenialReason InDenialReason)
 {
 	DenialReason = InDenialReason;
diff --git a/Engine/Plugins/Online/OnlineFramework/Source/Party/Public/Party/PartyTypes.h b/Engine/Plugins/Online/OnlineFramework/Source/Party/Public/Party/PartyTypes.h
--- a/Engine/Plugins/Online/OnlineFramework/Source/Party/Public/Party/PartyTypes.h
+++ b/Engine/Plugins/Online/OnlineFramework/Source/Party/Public/Party/PartyTypes.h
@@ -371,6 +371,7 @@ public:
 	FJoinPartyResult(FPartyJoinDenialReason InDenialReason);
 	FJoinPartyResult(EJoinPartyCompletionResult InResult);
 	FJoinPartyResult(EJoinPartyCompletionResult InResult, FPartyJoinDenialReason InDenialReason);
+	FJoinPartyResult(EJoinPartyCompletionResult InResult, int32 InResultSubCode);
 	
 	void SetDenialReason(FPartyJoinDenialReason InDenialReason);
 	void SetResult(EJoinPartyCompletionResult InResult);
@@ -379,10 +380,14 @@ public:
 
 	EJoinPartyCompletionResult GetResult() const { return Result; }
 	FPartyJoinDenialReason GetDenialReason() const { return DenialReason; }
+	int32 GetResultSubCode() const { return ResultSubCode; }
 
 private:
 	EJoinPartyCompletionResult Result;
+	/** Denial reason - used if Result is NotApproved */
 	FPartyJoinDenialReason DenialReason;
+	/** Result sub code - used for any other Result type */
+	int32 ResultSubCode = 0;
 };
 
 /** Base for all rep data structs */
