Adding scope counter to FSocialUserList::UpdateListInternal to identify FTicker_Tick time.

#rnx

diff --git a/Engine/Plugins/Online/OnlineFramework/Source/Party/Private/User/SocialUserList.cpp b/Engine/Plugins/Online/OnlineFramework/Source/Party/Private/User/SocialUserList.cpp
--- a/Engine/Plugins/Online/OnlineFramework/Source/Party/Private/User/SocialUserList.cpp
+++ b/Engine/Plugins/Online/OnlineFramework/Source/Party/Private/User/SocialUserList.cpp
@@ -484,6 +484,8 @@ bool FSocialUserList::HandleAutoUpdateList(float)
 
 void FSocialUserList::UpdateListInternal()
 {
+	QUICK_SCOPE_CYCLE_COUNTER(STAT_SocialUserList_UpdateList);
+
 	// Re-evaluate whether each user with dirtied presence is still fit for the list
 	for (TWeakObjectPtr<USocialUser> DirtyUser : UsersWithDirtyPresence)
 	{
