Add support for branch versioned hotfix files.
#Tested in a cooked server with Branch-Release-10.30_DefaultRuntimeOptions.ini and everything worked properly.

diff --git a/Engine/Plugins/Online/OnlineFramework/Source/Hotfix/Private/OnlineHotfixManager.cpp b/Engine/Plugins/Online/OnlineFramework/Source/Hotfix/Private/OnlineHotfixManager.cpp
--- a/Engine/Plugins/Online/OnlineFramework/Source/Hotfix/Private/OnlineHotfixManager.cpp
+++ b/Engine/Plugins/Online/OnlineFramework/Source/Hotfix/Private/OnlineHotfixManager.cpp
@@ -13,6 +13,7 @@
 
 #include "Misc/PackageName.h"
 #include "Misc/CommandLine.h"
+#include "Misc/EngineVersion.h"
 #include "Misc/FileHelper.h"
 #include "Misc/CoreDelegates.h"
 #include "Misc/App.h"
@@ -28,6 +29,8 @@ DEFINE_LOG_CATEGORY(LogHotfixManager);
 #define HOTFIX_SEPARATOR TEXT("_")
 /** The prefix for any hotfix file that expects to indicate version information */
 #define HOTFIX_VERSION_TAG TEXT("Ver-")
+/** The prefix for any hotfix file that expects to indicate branch version information */
+#define HOTFIX_BRANCH_VERSION_TAG TEXT("Branch-")
 
 FName NAME_HotfixManager(TEXT("HotfixManager"));
 
@@ -60,14 +63,26 @@ namespace
 		return NetVerStr;
 	}
 
+	/** @return the expected branch version for hotfix files determined at compile time */
+	FString GetBranchVersion()
+	{
+		static FString BranchVersion;
+		if (BranchVersion.IsEmpty())
+		{
+			BranchVersion = HOTFIX_BRANCH_VERSION_TAG + FPaths::GetCleanFilename(FEngineVersion::Current().GetBranch());
+		}
+		return BranchVersion;
+	}
+
 	/**
 	 * Given a hotfix file name, return the file name with version stripped out and exposed separately
 	 *
 	 * @param InFilename name of file to search for version information
 	 * @param OutFilename name with version information removed
-	 * @param OutVersion version of the hotfix file it present in the name
+	 * @param OutNetVersion version of the hotfix file it present in the name
+	 * @param OutBranchVersion version of the hotfix file it present in the name
 	 */
-	void GetFilenameAndVersion(const FString& InFilename, FString& OutFilename, FString& OutVersion)
+	void GetFilenameAndVersion(const FString& InFilename, FString& OutFilename, FString& OutNetVersion, FString& OutBranchVersion)
 	{
 		TArray<FString> FileParts;
 		int32 NumTokens = InFilename.ParseIntoArray(FileParts, HOTFIX_SEPARATOR);
@@ -77,7 +92,11 @@ namespace
 			{
 				if (FileParts[i].StartsWith(HOTFIX_VERSION_TAG))
 				{
-					OutVersion = FileParts[i];
+					OutNetVersion = FileParts[i];
+				}
+				else if (FileParts[i].StartsWith(HOTFIX_BRANCH_VERSION_TAG))
+				{
+					OutBranchVersion = FileParts[i];
 				}
 				else
 				{
@@ -103,22 +122,35 @@ namespace
 	 */
 	bool IsCompatibleHotfixFile(const FString& InFilename, FString& OutFilename)
 	{
-		bool bHasVersion = false;
-		bool bCompatibleHotfix = false;
-		const FString NetworkVersion = GetNetworkVersion();
-		FString OutVersion;
-		GetFilenameAndVersion(InFilename, OutFilename, OutVersion);
+		bool bHasNetVersion = false;
+		bool bCompatibleNetHotfix = false;
+		bool bHasBranchVersion = false;
+		bool bCompatibleBranchHotfix = false;
+		FString OutNetVersion;
+		FString OutBranchVersion;
+		GetFilenameAndVersion(InFilename, OutFilename, OutNetVersion, OutBranchVersion);
 
-		if (!OutVersion.IsEmpty())
+		if (!OutNetVersion.IsEmpty())
 		{
-			bHasVersion = true;
-			if (OutVersion == NetworkVersion)
+			bHasNetVersion = true;
+			const FString NetworkVersion = GetNetworkVersion();
+			if (OutNetVersion == NetworkVersion)
 			{
-				bCompatibleHotfix = true;
+				bCompatibleNetHotfix = true;
 			}
 		}
 
-		return bCompatibleHotfix || !bHasVersion;
+		if (!OutBranchVersion.IsEmpty())
+		{
+			bHasBranchVersion = true;
+			const FString BranchVersion = GetBranchVersion();
+			if (OutBranchVersion == BranchVersion)
+			{
+				bCompatibleBranchHotfix = true;
+			}
+		}
+
+		return (bCompatibleNetHotfix || !bHasNetVersion) && (bCompatibleBranchHotfix || !bHasBranchVersion);
 	}
 }
 
@@ -271,8 +303,8 @@ struct FHotfixFileSortPredicate
 
 			if (InHotfixName.EndsWith(TEXT("INI")))
 			{
-				FString HotfixName, NetworkVersion;
-				GetFilenameAndVersion(InHotfixName, HotfixName, NetworkVersion);
+				FString HotfixName, NetworkVersion, BranchVersion;
+				GetFilenameAndVersion(InHotfixName, HotfixName, NetworkVersion, BranchVersion);
 
 				// Defaults are applied first
 				if (HotfixName.StartsWith(DefaultPrefix))
@@ -295,9 +327,15 @@ struct FHotfixFileSortPredicate
 					Priority = 40;
 				}
 
+				if (!BranchVersion.IsEmpty())
+				{
+					// Branch versioned hotfixes apply after all but net versioned hotfixes within their type
+					Priority += 3;
+				}
+
 				if (!NetworkVersion.IsEmpty())
 				{
-					// Versioned hotfixes apply last within their type
+					// Network versioned hotfixes apply last within their type
 					Priority += 5;
 				}
 			}
@@ -748,7 +786,8 @@ FString UOnlineHotfixManager::GetStrippedConfigFileName(const FString& IniName)
 {
 	FString StrippedIniName;
 	FString NetworkVersion;
-	GetFilenameAndVersion(IniName, StrippedIniName, NetworkVersion);
+	FString BranchVersion;
+	GetFilenameAndVersion(IniName, StrippedIniName, NetworkVersion, BranchVersion);
 
 	if (StrippedIniName.StartsWith(PlatformPrefix))
 	{
