Added ISocialUserList::GetListName() accessor to SocialUserList to aid in logging/debugging.

diff --git a/Engine/Plugins/Online/OnlineFramework/Source/Party/Private/User/SocialUserList.h b/Engine/Plugins/Online/OnlineFramework/Source/Party/Private/User/SocialUserList.h
--- a/Engine/Plugins/Online/OnlineFramework/Source/Party/Private/User/SocialUserList.h
+++ b/Engine/Plugins/Online/OnlineFramework/Source/Party/Private/User/SocialUserList.h
@@ -18,6 +18,8 @@ public:
 	FOnUserRemoved& OnUserRemoved() const override { return OnUserRemovedEvent; }
 	FOnUpdateComplete& OnUpdateComplete() const override { return OnUpdateCompleteEvent; }
 
+	virtual FString GetListName() const override { return ListConfig.Name; }
+
 	void UpdateNow();
 	void SetAllowAutoUpdate(bool bIsEnabled) { bAllowAutoUpdate = bIsEnabled; }
 	void SetAutoUpdatePeriod(float InAutoUpdatePeriod);
diff --git a/Engine/Plugins/Online/OnlineFramework/Source/Party/Public/User/ISocialUserList.h b/Engine/Plugins/Online/OnlineFramework/Source/Party/Public/User/ISocialUserList.h
--- a/Engine/Plugins/Online/OnlineFramework/Source/Party/Public/User/ISocialUserList.h
+++ b/Engine/Plugins/Online/OnlineFramework/Source/Party/Public/User/ISocialUserList.h
@@ -63,6 +63,8 @@ public:
 
 	virtual const TArray<USocialUser*>& GetUsers() const = 0;
 
+	virtual FString GetListName() const = 0;
+
 	/** Trigger an update of the list immediately, regardless of auto update period */
 	virtual void UpdateNow() = 0;
 
