Removing unnecessarily verbose logging from SocialToolkit

diff --git a/Engine/Plugins/Online/OnlineFramework/Source/Party/Private/SocialToolkit.cpp b/Engine/Plugins/Online/OnlineFramework/Source/Party/Private/SocialToolkit.cpp
--- a/Engine/Plugins/Online/OnlineFramework/Source/Party/Private/SocialToolkit.cpp
+++ b/Engine/Plugins/Online/OnlineFramework/Source/Party/Private/SocialToolkit.cpp
@@ -743,24 +743,7 @@ void USocialToolkit::HandlePresenceReceived(const FUniqueNetId& UserId, const TS
 	}
 	else if (SubsystemType == ESocialSubsystem::Platform)
 	{
-		QUICK_SCOPE_CYCLE_COUNTER(STAT_USocialToolkit_HandlePresenceReceived_Error);
-		FString ErrorString = TEXT("Platform presence received, but existing SocialUser could not be found.\n");
-		ErrorString += TEXT("Incoming UserId is ") + UserId.ToString() + TEXT(", as a UniqueIdRepl it's ") + FUniqueNetIdRepl(UserId).ToString();
-
-		ErrorString += TEXT("Outputting all cached platform IDs and the corresponding user: \n") + UserId.ToString();
-		for (auto IdUserPair : UsersBySubsystemIds)
-		{
-			if (IdUserPair.Key.GetType() != MCP_SUBSYSTEM)
-			{
-				ErrorString += FString::Printf(TEXT("\tUserId [%s]: SocialUser [%s]\n"), *IdUserPair.Key.ToString(), *IdUserPair.Value->ToDebugString());
-				if (IdUserPair.Key == FUniqueNetIdRepl(UserId) || !ensure(*IdUserPair.Key != UserId))
-				{
-					ErrorString += TEXT("\t\tAnd look at that, this one DOES actually match. The map has lied to us!!\n");
-				}
-			}
-		}
-
-		UE_LOG(LogParty, Error, TEXT("%s"), *ErrorString);
+		UE_LOG(LogParty, Error, TEXT("Platform presence received for UserId [%s], but existing SocialUser could not be found.\n"), *UserId.ToDebugString());
 	}
 }
 
