Removed SetAutoUpdatePeriod from SocialUserList to help clarify the process of enabling/disabling auto update.

#tests Verified that SetAutoUpdate correctly enables/disables the list ticking process.

diff --git a/Engine/Plugins/Online/OnlineFramework/Source/Party/Private/User/SocialUserList.cpp b/Engine/Plugins/Online/OnlineFramework/Source/Party/Private/User/SocialUserList.cpp
--- a/Engine/Plugins/Online/OnlineFramework/Source/Party/Private/User/SocialUserList.cpp
+++ b/Engine/Plugins/Online/OnlineFramework/Source/Party/Private/User/SocialUserList.cpp
@@ -80,10 +80,8 @@ void FSocialUserList::InitializeList()
 		}
 	}
 
-	if (ListConfig.bAutoUpdate)
-	{
-		SetAutoUpdatePeriod(USocialSettings::GetUserListAutoUpdateRate());
-	}
+	AutoUpdatePeriod = USocialSettings::GetUserListAutoUpdateRate();
+	SetAllowAutoUpdate(ListConfig.bAutoUpdate);
 }
 
 void FSocialUserList::AddReferencedObjects(FReferenceCollector& Collector)
@@ -96,17 +94,16 @@ void FSocialUserList::UpdateNow()
 	UpdateListInternal();
 }
 
-void FSocialUserList::SetAutoUpdatePeriod(float InAutoUpdatePeriod)
+void FSocialUserList::SetAllowAutoUpdate(bool bIsEnabled)
 {
-	AutoUpdatePeriod = InAutoUpdatePeriod;
 	if (UpdateTickerHandle.IsValid())
 	{
 		FTicker::GetCoreTicker().RemoveTicker(UpdateTickerHandle);
 	}
-
-	if (InAutoUpdatePeriod >= 0.f)
+	
+	if (bIsEnabled)
 	{
-		UpdateTickerHandle = FTicker::GetCoreTicker().AddTicker(FTickerDelegate::CreateSP(this, &FSocialUserList::HandleAutoUpdateList), InAutoUpdatePeriod);
+		UpdateTickerHandle = FTicker::GetCoreTicker().AddTicker(FTickerDelegate::CreateSP(this, &FSocialUserList::HandleAutoUpdateList), AutoUpdatePeriod);
 	}
 }
 
@@ -475,10 +472,7 @@ struct FUserSortData
 
 bool FSocialUserList::HandleAutoUpdateList(float)
 {
-	if (bAllowAutoUpdate)
-	{
-		UpdateListInternal();
-	}
+	UpdateListInternal();
 	return true;
 }
 
diff --git a/Engine/Plugins/Online/OnlineFramework/Source/Party/Private/User/SocialUserList.h b/Engine/Plugins/Online/OnlineFramework/Source/Party/Private/User/SocialUserList.h
--- a/Engine/Plugins/Online/OnlineFramework/Source/Party/Private/User/SocialUserList.h
+++ b/Engine/Plugins/Online/OnlineFramework/Source/Party/Private/User/SocialUserList.h
@@ -21,8 +21,7 @@ public:
 	virtual FString GetListName() const override { return ListConfig.Name; }
 
 	void UpdateNow();
-	void SetAllowAutoUpdate(bool bIsEnabled) { bAllowAutoUpdate = bIsEnabled; }
-	void SetAutoUpdatePeriod(float InAutoUpdatePeriod);
+	void SetAllowAutoUpdate(bool bIsEnabled);
 	const TArray<USocialUser*>& GetUsers() const { return Users; }
 
 	bool HasPresenceFilters() const;
diff --git a/Engine/Plugins/Online/OnlineFramework/Source/Party/Public/User/ISocialUserList.h b/Engine/Plugins/Online/OnlineFramework/Source/Party/Public/User/ISocialUserList.h
--- a/Engine/Plugins/Online/OnlineFramework/Source/Party/Public/User/ISocialUserList.h
+++ b/Engine/Plugins/Online/OnlineFramework/Source/Party/Public/User/ISocialUserList.h
@@ -70,7 +70,4 @@ public:
 
 	/** Give external overwrite to disable list auto update for perf */
 	virtual void SetAllowAutoUpdate(bool bIsEnabled) = 0;
-
-	/** Sets the period at which to update the list with all users that  */
-	virtual void SetAutoUpdatePeriod(float InAutoUpdatePeriod) = 0;
 };
\ No newline at end of file
