Add convenience method for adding properties to presence

diff --git a/Engine/Plugins/Online/OnlineFramework/Source/Party/Private/SocialToolkit.cpp b/Engine/Plugins/Online/OnlineFramework/Source/Party/Private/SocialToolkit.cpp
--- a/Engine/Plugins/Online/OnlineFramework/Source/Party/Private/SocialToolkit.cpp
+++ b/Engine/Plugins/Online/OnlineFramework/Source/Party/Private/SocialToolkit.cpp
@@ -222,6 +222,33 @@ void USocialToolkit::SetLocalUserOnlineState(EOnlinePresenceState::Type OnlineSt
 	}
 }
 
+void USocialToolkit::AddLocalUserOnlineProperties(FPresenceProperties OnlineProperties)
+{
+	if (IOnlineSubsystem* PrimaryOss = GetSocialOss(ESocialSubsystem::Primary))
+	{
+		IOnlinePresencePtr PresenceInterface = PrimaryOss->GetPresenceInterface();
+		FUniqueNetIdRepl LocalUserId = GetLocalUserNetId(ESocialSubsystem::Primary);
+		if (PresenceInterface.IsValid() && LocalUserId.IsValid())
+		{
+			TSharedPtr<FOnlineUserPresence> CurrentPresence;
+			PresenceInterface->GetCachedPresence(*LocalUserId, CurrentPresence);
+
+			FOnlineUserPresenceStatus NewStatus;
+			if (CurrentPresence.IsValid())
+			{
+				NewStatus = CurrentPresence->Status;
+			}
+			
+			for (TPair<FPresenceKey, FVariantData>& Pair : OnlineProperties)
+			{
+				NewStatus.Properties.Emplace(MoveTemp(Pair.Key), MoveTemp(Pair.Value));
+			}
+
+			PresenceInterface->SetPresence(*LocalUserId, NewStatus);
+		}
+	}
+}
+
 USocialManager& USocialToolkit::GetSocialManager() const
 {
 	USocialManager* OuterSocialManager = GetTypedOuter<USocialManager>();
diff --git a/Engine/Plugins/Online/OnlineFramework/Source/Party/Public/SocialToolkit.h b/Engine/Plugins/Online/OnlineFramework/Source/Party/Public/SocialToolkit.h
--- a/Engine/Plugins/Online/OnlineFramework/Source/Party/Public/SocialToolkit.h
+++ b/Engine/Plugins/Online/OnlineFramework/Source/Party/Public/SocialToolkit.h
@@ -6,6 +6,7 @@
 #include "SocialTypes.h"
 #include "User/ISocialUserList.h"
 #include "OnlineSubsystem.h"
+#include "Interfaces/OnlinePresenceInterface.h"
 
 #include "SocialToolkit.generated.h"
 
@@ -27,8 +28,6 @@ namespace FriendInviteFailureReason
 	const FString InviteFailReason_AlreadyFriends = TEXT("AlreadyFriends");
 }
 
-namespace EOnlinePresenceState { enum Type : uint8; }
-
 DECLARE_DELEGATE_OneParam(FUserDependentAction, USocialUser&);
 
 /** Represents the full suite of social functionality available to a given LocalPlayer */
@@ -63,6 +62,7 @@ public:
 
 	const FOnlineUserPresence* GetPresenceInfo(ESocialSubsystem SubsystemType) const;
 	void SetLocalUserOnlineState(EOnlinePresenceState::Type OnlineState);
+	void AddLocalUserOnlineProperties(FPresenceProperties OnlineProperties);
 
 	USocialManager& GetSocialManager() const;
 	ULocalPlayer& GetOwningLocalPlayer() const;
