Modified the way we determine a SocialUser's Platform for GetCurrentPlatform calls by using best available information (by also searching the local player's persistent party).  This function is now virtual to allow games to further dig deeper if desired.

Added GetPlatformIconMarkupTag with an accompanying ruleset to allow crossplay logic to be applied, including rules on when to display an icon depending on the local player's crossplay situation.

#tests Compile DebugGame Editor Win64 / Shipping Client / Multiple -game clients with various -platformtest overrides running matches while in and out of a party, emulating different crossplay scenarios.

diff --git a/Engine/Plugins/Online/OnlineFramework/Source/Party/Private/User/SocialUser.cpp b/Engine/Plugins/Online/OnlineFramework/Source/Party/Private/User/SocialUser.cpp
--- a/Engine/Plugins/Online/OnlineFramework/Source/Party/Private/User/SocialUser.cpp
+++ b/Engine/Plugins/Online/OnlineFramework/Source/Party/Private/User/SocialUser.cpp
@@ -544,11 +544,84 @@ FUserPlatform USocialUser::GetCurrentPlatform() const
 		// We have no platform presence, but we do have primary, so get the platform from that
 		return FUserPlatform(PrimaryPresence->GetPlatform());
 	}
-	
+
+	UE_LOG(LogOnline, VeryVerbose, TEXT("*!* %s - No Presence found for user, searching Party..."), ANSI_TO_TCHAR(__FUNCTION__));
+	if (USocialParty* Party = GetOwningToolkit().GetSocialManager().GetPersistentParty())
+	{
+		if (UPartyMember* PartyMember = Party->GetPartyMember(GetUserId(ESocialSubsystem::Primary)))
+		{
+			FUserPlatform Platform = PartyMember->GetRepData().GetPlatform();
+			UE_LOG(LogOnline, VeryVerbose, TEXT("*!* %s - Party Member Found for user! RepDataPlatform: %s"), ANSI_TO_TCHAR(__FUNCTION__), *Platform.ToString());
+			if (Platform.IsValid())
+			{
+				return Platform;
+			}
+		}
+	}
+
 	// We don't have any presence for this user (or we do and they're offline) and they aren't the local player, so we really don't have any idea what their current platform is
 	return FUserPlatform();
 }
 
+FString USocialUser::GetPlatformIconMarkupTag(EPlatformIconDisplayRule DisplayRule) const
+{
+	FString DummyString;
+	return GetPlatformIconMarkupTag(DisplayRule, DummyString);
+}
+
+FString USocialUser::GetPlatformIconMarkupTag(EPlatformIconDisplayRule DisplayRule, FString& OutLegacyString) const
+{
+	const FUserPlatform UserPlatform = GetCurrentPlatform();
+	const FUserPlatform LocalPlatform = FUserPlatform(IOnlineSubsystem::GetLocalPlatformName());
+
+	OutLegacyString = UserPlatform;
+	switch (DisplayRule)
+	{
+	case EPlatformIconDisplayRule::Always:
+		UE_LOG(LogOnline, VeryVerbose, TEXT("*!*    %s - User: %s - Returning TypeName %s due to ALWAYS"), ANSI_TO_TCHAR(__FUNCTION__), *GetDisplayName(), *UserPlatform.GetTypeName());
+		return UserPlatform.GetTypeName();
+	case EPlatformIconDisplayRule::AlwaysIfDifferent:
+		UE_LOG(LogOnline, VeryVerbose, TEXT("*!*    %s - User: %s - CrossplayLocalPlatform? %d Returning %s due to ALWAYSIFDIFFERENT"), ANSI_TO_TCHAR(__FUNCTION__), *GetDisplayName(), UserPlatform.IsCrossplayWithLocalPlatform(), UserPlatform.IsCrossplayWithLocalPlatform() ? *UserPlatform.GetTypeName() : TEXT(""));
+		if (!UserPlatform.IsCrossplayWithLocalPlatform())
+		{
+			OutLegacyString = TEXT("");
+			return TEXT("");
+		}
+		return UserPlatform.GetTypeName();
+	case EPlatformIconDisplayRule::Never:
+		UE_LOG(LogOnline, VeryVerbose, TEXT("*!*    %s - User: %s - Returning nothing due to NEVER"), ANSI_TO_TCHAR(__FUNCTION__), *GetDisplayName());
+		OutLegacyString = TEXT("");
+		return TEXT("");
+	}
+
+	if (DisplayRule == EPlatformIconDisplayRule::AlwaysWhenInCrossplayParty ||
+		DisplayRule == EPlatformIconDisplayRule::AlwaysIfDifferentWhenInCrossplayParty)
+	{
+		if (USocialParty* Party = GetOwningToolkit().GetSocialManager().GetPersistentParty())
+		{
+			if (!Party->IsCurrentlyCrossplaying())
+			{
+				UE_LOG(LogOnline, VeryVerbose, TEXT("*!*    %s - User: %s - Returning nothing due to Not Being In Crossplay with WHENINCROSSPLAYPARTY specified"), ANSI_TO_TCHAR(__FUNCTION__), *GetDisplayName());
+				OutLegacyString = TEXT("");
+				return TEXT("");
+			}
+		}
+	}
+
+	const bool bIsCrossplayWithMe = UserPlatform.IsCrossplayWithLocalPlatform();
+	if (DisplayRule == EPlatformIconDisplayRule::AlwaysWhenInCrossplayParty ||
+		(DisplayRule == EPlatformIconDisplayRule::AlwaysIfDifferentWhenInCrossplayParty && bIsCrossplayWithMe))
+	{
+		FString TypeName = UserPlatform.GetTypeName();
+		UE_LOG(LogOnline, VeryVerbose, TEXT("*!*    %s - User: %s - Returning TypeName %s, DisplayRule: %d, bIsCrossplayWithMe: %d"), ANSI_TO_TCHAR(__FUNCTION__), *GetDisplayName(), *TypeName, (int32)DisplayRule, bIsCrossplayWithMe);
+		return TypeName;
+	}
+
+	UE_LOG(LogOnline, VeryVerbose, TEXT("*!*    %s - User: %s - Returning nothing, likely due to AlwaysIfDifferentWhenInCrossplayParty and not being different, being in a Crossplay Party, and not being Different"), ANSI_TO_TCHAR(__FUNCTION__), *GetDisplayName());
+	OutLegacyString = TEXT("");
+	return TEXT("");
+}
+
 void USocialUser::GetRichPresenceText(FText& OutRichPresence) const
 {
 	if (IsBlocked())
diff --git a/Engine/Plugins/Online/OnlineFramework/Source/Party/Public/SocialTypes.h b/Engine/Plugins/Online/OnlineFramework/Source/Party/Public/SocialTypes.h
--- a/Engine/Plugins/Online/OnlineFramework/Source/Party/Public/SocialTypes.h
+++ b/Engine/Plugins/Online/OnlineFramework/Source/Party/Public/SocialTypes.h
@@ -163,6 +163,16 @@ private:
 	TArray<FSocialActionStep> ActionSteps;
 };
 
+UENUM()
+enum class EPlatformIconDisplayRule : uint8
+{
+	Always,									// always show the platform icon
+	AlwaysIfDifferent,						// always show the icon if it's a different platform from my own
+	AlwaysWhenInCrossplayParty,				// always show the icon if I'm in a crossplay party
+	AlwaysIfDifferentWhenInCrossplayParty,	// only show the icon if it's different from my own and I'm in a crossplay party
+	Never,									// never show the icon
+};
+
 #define DECLARE_SHARED_PTR_ALIASES(Type)	\
 class Type;	\
 using Type##Ptr = TSharedPtr<Type>;	\
diff --git a/Engine/Plugins/Online/OnlineFramework/Source/Party/Public/User/SocialUser.h b/Engine/Plugins/Online/OnlineFramework/Source/Party/Public/User/SocialUser.h
--- a/Engine/Plugins/Online/OnlineFramework/Source/Party/Public/User/SocialUser.h
+++ b/Engine/Plugins/Online/OnlineFramework/Source/Party/Public/User/SocialUser.h
@@ -13,6 +13,7 @@
 class IOnlinePartyJoinInfo;
 class FOnlineUserPresence;
 class UPartyMember;
+enum class EPlatformIconDisplayRule : uint8;
 
 namespace EOnlinePresenceState { enum Type : uint8; }
 
@@ -51,7 +52,10 @@ public:
 	const FOnlineUserPresence* GetFriendPresenceInfo(ESocialSubsystem SubsystemType) const;
 	FDateTime GetFriendshipCreationDate() const;
 	FText GetSocialName() const;
-	FUserPlatform GetCurrentPlatform() const;
+	virtual FUserPlatform GetCurrentPlatform() const;
+
+	FString GetPlatformIconMarkupTag(EPlatformIconDisplayRule DisplayRule) const;
+	virtual FString GetPlatformIconMarkupTag(EPlatformIconDisplayRule DisplayRule, FString& OutLegacyString) const;
 
 	virtual void GetRichPresenceText(FText& OutRichPresence) const;
 
