Modification of login flow so that secondary player automatically joins the primary player's party instead of creating one

diff --git a/Engine/Plugins/Online/OnlineFramework/Source/Party/Private/Party/PartyMember.cpp b/Engine/Plugins/Online/OnlineFramework/Source/Party/Private/Party/PartyMember.cpp
--- a/Engine/Plugins/Online/OnlineFramework/Source/Party/Private/Party/PartyMember.cpp
+++ b/Engine/Plugins/Online/OnlineFramework/Source/Party/Private/Party/PartyMember.cpp
@@ -180,7 +180,7 @@ bool UPartyMember::IsPartyLeader() const
 
 bool UPartyMember::IsLocalPlayer() const
 {
-	return GetPrimaryNetId() == GetParty().GetOwningLocalUserId();
+	return GetParty().GetSocialManager().IsLocalUser(GetPrimaryNetId(), ESocialSubsystem::Primary);
 }
 
 void UPartyMember::NotifyMemberDataReceived(const FOnlinePartyData& MemberData)
diff --git a/Engine/Plugins/Online/OnlineFramework/Source/Party/Private/Party/SocialParty.cpp b/Engine/Plugins/Online/OnlineFramework/Source/Party/Private/Party/SocialParty.cpp
--- a/Engine/Plugins/Online/OnlineFramework/Source/Party/Private/Party/SocialParty.cpp
+++ b/Engine/Plugins/Online/OnlineFramework/Source/Party/Private/Party/SocialParty.cpp
@@ -553,8 +553,8 @@ UPartyMember* USocialParty::GetOrCreatePartyMember(const FUniqueNetId& MemberId)
 		}
 		else
 		{
-			//@todo DanH Splitscreen: Multiple members in the party can still be local players #future
-			TSubclassOf<UPartyMember> PartyMemberClass = GetDesiredMemberClass(MemberId == *OwningLocalUserId);
+			const bool bIsLocalUser = GetSocialManager().IsLocalUser(MemberId, ESocialSubsystem::Primary);
+			TSubclassOf<UPartyMember> PartyMemberClass = GetDesiredMemberClass(bIsLocalUser);
 			if (ensure(PartyMemberClass))
 			{
 				const FOnlinePartyId& PartyId = GetPartyId();
diff --git a/Engine/Plugins/Online/OnlineFramework/Source/Party/Private/SocialManager.cpp b/Engine/Plugins/Online/OnlineFramework/Source/Party/Private/SocialManager.cpp
--- a/Engine/Plugins/Online/OnlineFramework/Source/Party/Private/SocialManager.cpp
+++ b/Engine/Plugins/Online/OnlineFramework/Source/Party/Private/SocialManager.cpp
@@ -325,6 +325,20 @@ FUniqueNetIdRepl USocialManager::GetFirstLocalUserId(ESocialSubsystem SubsystemT
 	return FUniqueNetIdRepl();
 }
 
+
+bool USocialManager::IsLocalUser(const FUniqueNetIdRepl& LocalUserId, ESocialSubsystem SubsystemType) const
+{
+	for(USocialToolkit* SocialToolkit : SocialToolkits)
+	{
+		if (SocialToolkit->GetLocalUserNetId(SubsystemType) == LocalUserId)
+		{
+			return true;
+		}
+	}
+
+	return false;
+}
+
 int32 USocialManager::GetFirstLocalUserNum() const
 {
 	if (SocialToolkits.Num() > 0)
@@ -649,6 +663,30 @@ USocialToolkit& USocialManager::CreateSocialToolkit(ULocalPlayer& OwningLocalPla
 	return *NewToolkit;
 }
 
+void USocialManager::RegisterSecondaryPlayer(int32 LocalPlayerNum, const FOnJoinPartyComplete& Delegate)
+{
+	USocialToolkit* SocialToolkit = GetSocialToolkit(LocalPlayerNum);
+	IOnlinePartyPtr PartyInterface = Online::GetPartyInterfaceChecked(GetWorld());
+	USocialParty* PersistentParty = GetPersistentParty();
+
+	if (PartyInterface != nullptr && PersistentParty != nullptr)
+	{
+		FUniqueNetIdRepl PrimaryUserId = GetFirstLocalUserId(ESocialSubsystem::Primary);
+		FUniqueNetIdRepl SecondaryUserId = SocialToolkit->GetLocalUser().GetUserId(ESocialSubsystem::Primary);
+
+		if (PrimaryUserId.IsValid() && SecondaryUserId.IsValid())
+		{
+			FString JoinInfoStr = PartyInterface->MakeJoinInfoJson(*PrimaryUserId, PersistentParty->GetPartyId());
+			IOnlinePartyJoinInfoConstPtr JoinInfo = PartyInterface->MakeJoinInfoFromJson(JoinInfoStr);
+
+			if (JoinInfo->IsValid())
+			{
+				PartyInterface->JoinParty(*SecondaryUserId, *JoinInfo, Delegate);
+			}
+		}
+	}
+}
+
 void USocialManager::QueryPartyJoinabilityInternal(FJoinPartyAttempt& JoinAttempt)
 {
 	FUniqueNetIdRepl LocalUserId = GetFirstLocalUserId(ESocialSubsystem::Primary);
diff --git a/Engine/Plugins/Online/OnlineFramework/Source/Party/Public/SocialManager.h b/Engine/Plugins/Online/OnlineFramework/Source/Party/Public/SocialManager.h
--- a/Engine/Plugins/Online/OnlineFramework/Source/Party/Public/SocialManager.h
+++ b/Engine/Plugins/Online/OnlineFramework/Source/Party/Public/SocialManager.h
@@ -52,6 +52,7 @@ public:
 	USocialToolkit& GetSocialToolkit(ULocalPlayer& LocalPlayer) const;
 	USocialToolkit* GetFirstLocalUserToolkit() const;
 	FUniqueNetIdRepl GetFirstLocalUserId(ESocialSubsystem SubsystemType) const;
+	bool IsLocalUser(const FUniqueNetIdRepl& LocalUserId, ESocialSubsystem SubsystemType) const;
 	int32 GetFirstLocalUserNum() const;
 	USocialDebugTools* GetDebugTools() const;
 
@@ -94,6 +95,13 @@ public:
 
 	void HandlePartyDisconnected(USocialParty* LeavingParty);
 
+	/**
+	 * Makes an attempt for the target local player to join the primary local player's party
+	 * @param LocalPlayerNum - ControllerId of the Secondary player that wants to join the party
+	 * @param Delegate - Delegate run when the join process is finished
+	 */
+	void RegisterSecondaryPlayer(int32 LocalPlayerNum, const FOnJoinPartyComplete& Delegate = FOnJoinPartyComplete());
+
 PARTY_SCOPE:
 	/** Validates that the target user has valid join info for us to use and that we can join any party of the given type */
 	FJoinPartyResult ValidateJoinTarget(const USocialUser& UserToJoin, const FOnlinePartyTypeId& PartyTypeId) const;
@@ -103,7 +111,7 @@ PARTY_SCOPE:
 
 	void NotifyPartyInitialized(USocialParty& Party);
 	USocialToolkit* GetSocialToolkit(int32 LocalPlayerNum) const;
-	
+
 protected:
 	struct PARTY_API FRejoinableParty : public TSharedFromThis<FRejoinableParty>
 	{
