Don't restart the UserList timer if we repeatedly call Allow(true)

#tests Compile DebugGame Editor Win64

diff --git a/Engine/Plugins/Online/OnlineFramework/Source/Party/Private/User/SocialUserList.cpp b/Engine/Plugins/Online/OnlineFramework/Source/Party/Private/User/SocialUserList.cpp
--- a/Engine/Plugins/Online/OnlineFramework/Source/Party/Private/User/SocialUserList.cpp
+++ b/Engine/Plugins/Online/OnlineFramework/Source/Party/Private/User/SocialUserList.cpp
@@ -96,12 +96,11 @@ void FSocialUserList::UpdateNow()
 
 void FSocialUserList::SetAllowAutoUpdate(bool bIsEnabled)
 {
-	if (UpdateTickerHandle.IsValid())
+	if (!bIsEnabled && UpdateTickerHandle.IsValid())
 	{
 		FTicker::GetCoreTicker().RemoveTicker(UpdateTickerHandle);
 	}
-	
-	if (bIsEnabled)
+	else if (bIsEnabled && !UpdateTickerHandle.IsValid())
 	{
 		UpdateTickerHandle = FTicker::GetCoreTicker().AddTicker(FTickerDelegate::CreateSP(this, &FSocialUserList::HandleAutoUpdateList), AutoUpdatePeriod);
 	}
