Fix for UpdateRequired message during Intro Matchmaking

diff --git a/Engine/Plugins/Online/OnlineFramework/Source/Hotfix/Private/UpdateManager.cpp b/Engine/Plugins/Online/OnlineFramework/Source/Hotfix/Private/UpdateManager.cpp
--- a/Engine/Plugins/Online/OnlineFramework/Source/Hotfix/Private/UpdateManager.cpp
+++ b/Engine/Plugins/Online/OnlineFramework/Source/Hotfix/Private/UpdateManager.cpp
@@ -208,7 +208,7 @@ void UUpdateManager::CheckComplete(EUpdateCompletionStatus Result, bool bUpdateT
 		LastUpdateCheck[bCheckHotfixAvailabilityOnly] = FDateTime::UtcNow();
 	}
 
-	auto CompletionDelegate = [this, Result]()
+	auto CompletionDelegate = [this, Result, bUpdateTimestamp]()
 	{
 		UE_LOG(LogHotfixManager, Display, TEXT("External CheckComplete %s"), UpdateCompletionEnum ? *UpdateCompletionEnum->GetNameStringByValue((int64)Result) : TEXT("Invalid"));
 		if (!bInitialUpdateFinished)
@@ -222,8 +222,17 @@ void UUpdateManager::CheckComplete(EUpdateCompletionStatus Result, bool bUpdateT
 			SetUpdateState(EUpdateState::UpdateComplete);
 		}
 
+		EUpdateCompletionStatus FinalResult = Result;
+		if (Result == EUpdateCompletionStatus::UpdateSuccess && !bCheckHotfixAvailabilityOnly && !bUpdateTimestamp)
+		{
+			// if this is a cached value, and we are not checking availability only, we should return NoChange,
+			// As we have already applied this hotfix.
+			FinalResult = EUpdateCompletionStatus::UpdateSuccess_NoChange;
+		}
+
 		bCheckHotfixAvailabilityOnly = false;
-		OnUpdateCheckComplete().Broadcast(Result);
+		
+		OnUpdateCheckComplete().Broadcast(FinalResult);
 	};
 
 	// Delay completion delegate to give UI a chance to show the screen for a reasonable amount of time
