Adding getter and setter in SocialUser for nicknames.

diff --git a/Engine/Plugins/Online/OnlineFramework/Source/Party/Private/User/SocialUser.cpp b/Engine/Plugins/Online/OnlineFramework/Source/Party/Private/User/SocialUser.cpp
--- a/Engine/Plugins/Online/OnlineFramework/Source/Party/Private/User/SocialUser.cpp
+++ b/Engine/Plugins/Online/OnlineFramework/Source/Party/Private/User/SocialUser.cpp
@@ -399,6 +399,49 @@ FString USocialUser::GetDisplayName(ESocialSubsystem SubsystemType) const
 	return FString();
 }
 
+FString USocialUser::GetNickname() const
+{
+	if (IsFriend())
+	{
+#if WITH_EDITOR
+		bool bRandomNickname = FParse::Param(FCommandLine::Get(), TEXT("RandomNickname"));
+		if (bRandomNickname && !IsLocalUser())
+		{
+			if (GetTypeHash(GetUserId(ESocialSubsystem::Primary)) % 2)
+			{
+				FString RandomNickname = GetDisplayName(ESocialSubsystem::Primary);
+				for (int i = 1; i < RandomNickname.Len(); i += 2)
+				{
+					RandomNickname.RemoveAt(i);
+				}
+				return RandomNickname;
+			}
+		}
+#endif
+
+		if (const FSubsystemUserInfo* SubsystemInfo = SubsystemInfoByType.Find(ESocialSubsystem::Primary))
+		{
+			if (SubsystemInfo->FriendInfo.IsValid())
+			{
+				TSharedPtr<FOnlineFriend> OnlineFriend = SubsystemInfo->FriendInfo.Pin();
+				FString Nickname;
+				OnlineFriend->GetUserAttribute(USER_ATTR_ALIAS, Nickname);
+				return Nickname;
+			}
+		}
+	}
+	return TEXT("");
+}
+
+void USocialUser::SetNickname(const FString& InNickname)
+{
+	// @note StephanJ: Nickname for now only applies to primary/MCP subsystem.
+	IOnlineFriendsPtr FriendsInterface = GetOwningToolkit().GetSocialOss(ESocialSubsystem::Primary)->GetFriendsInterface();
+	check(FriendsInterface.IsValid());
+
+	FriendsInterface->SetFriendAlias(GetOwningToolkit().GetLocalUserNum(), *GetUserId(ESocialSubsystem::Primary), EFriendsLists::ToString(EFriendsLists::Default), InNickname, FOnSetFriendAliasComplete::CreateUObject(const_cast<USocialUser*>(this), &USocialUser::HandleSetNicknameComplete));
+}
+
 EInviteStatus::Type USocialUser::GetFriendInviteStatus(ESocialSubsystem SubsystemType) const
 {
 	if (const FSubsystemUserInfo* SubsystemInfo = SubsystemInfoByType.Find(SubsystemType))
@@ -1239,6 +1282,18 @@ void USocialUser::HandleQueryUserInfoComplete(ESocialSubsystem SubsystemType, bo
 	TryBroadcastInitializationComplete();
 }
 
+void USocialUser::HandleSetNicknameComplete(int32 LocalUserNum, const FUniqueNetId& FriendId, const FString& ListName, const FOnlineError& Error)
+{
+	if (!Error.WasSuccessful())
+	{
+		UE_LOG(LogOnline, Warning, TEXT("Set nickname request failed for user: %s"), *FriendId.ToDebugString());
+	}
+	else
+	{
+		OnNicknameChanged().Broadcast();
+	}
+}
+
 FString USocialUser::SanitizePresenceString(FString InString) const
 {
 	return InString;
diff --git a/Engine/Plugins/Online/OnlineFramework/Source/Party/Public/User/SocialUser.h b/Engine/Plugins/Online/OnlineFramework/Source/Party/Public/User/SocialUser.h
--- a/Engine/Plugins/Online/OnlineFramework/Source/Party/Public/User/SocialUser.h
+++ b/Engine/Plugins/Online/OnlineFramework/Source/Party/Public/User/SocialUser.h
@@ -15,6 +15,8 @@ class FOnlineUserPresence;
 class UPartyMember;
 enum class EPlatformIconDisplayRule : uint8;
 
+struct FOnlineError;
+
 namespace EOnlinePresenceState { enum Type : uint8; }
 
 DECLARE_MULTICAST_DELEGATE_OneParam(FOnNewSocialUserInitialized, USocialUser&);
@@ -44,6 +46,8 @@ public:
 	FUniqueNetIdRepl GetUserId(ESocialSubsystem SubsystemType) const;
 	FString GetDisplayName() const;
 	FString GetDisplayName(ESocialSubsystem SubsystemType) const;
+	FString GetNickname() const;
+	void SetNickname(const FString& InNickName);
 
 	EInviteStatus::Type GetFriendInviteStatus(ESocialSubsystem SubsystemType) const;
 	bool IsFriend() const;
@@ -101,6 +105,9 @@ public:
 
 	UPartyMember* GetPartyMember(const FOnlinePartyTypeId& PartyTypeId) const;
 
+	DECLARE_EVENT(USocialUser, FOnNicknameChanged);
+	FOnNicknameChanged& OnNicknameChanged() const { return OnNicknameChangedEvent; }
+
 	DECLARE_EVENT(USocialUser, FPartyInviteResponseEvent);
 	FPartyInviteResponseEvent& OnPartyInviteAccepted() const { return OnPartyInviteAcceptedEvent; }
 	FPartyInviteResponseEvent& OnPartyInviteRejected() const { return OnPartyInviteRejectedEvent; }
@@ -152,6 +159,7 @@ private:
 	
 	void SetUserInfo(ESocialSubsystem SubsystemType, const TSharedRef<FOnlineUser>& UserInfo);
 	void HandleQueryUserInfoComplete(ESocialSubsystem SubsystemType, bool bWasSuccessful, const TSharedPtr<FOnlineUser>& UserInfo);
+	void HandleSetNicknameComplete(int32 LocalUserNum, const FUniqueNetId& FriendId, const FString& ListName, const FOnlineError& Error);
 
 	virtual FString SanitizePresenceString(FString InString) const;
 	
@@ -190,6 +198,7 @@ private:
 	// Initialization delegates that fire only when a specific user has finishing initializing
 	static TMap<TWeakObjectPtr<USocialUser>, FOnNewSocialUserInitialized> InitEventsByUser;
 
+	mutable FOnNicknameChanged OnNicknameChangedEvent;
 	mutable FPartyInviteResponseEvent OnPartyInviteAcceptedEvent;
 	mutable FPartyInviteResponseEvent OnPartyInviteRejectedEvent;
 	mutable FOnUserPresenceChanged OnUserPresenceChangedEvent;
