Move bLeavePartyOnDisconnect to USocialSettings to be more consistent

diff --git a/Engine/Plugins/Online/OnlineFramework/Source/Party/Private/SocialManager.cpp b/Engine/Plugins/Online/OnlineFramework/Source/Party/Private/SocialManager.cpp
--- a/Engine/Plugins/Online/OnlineFramework/Source/Party/Private/SocialManager.cpp
+++ b/Engine/Plugins/Online/OnlineFramework/Source/Party/Private/SocialManager.cpp
@@ -1040,7 +1040,7 @@ void USocialManager::HandlePersistentPartyStateChanged(EPartyState NewState, EPa
 		
 		if (PreviousState == EPartyState::Active)
 		{
-			if (bLeavePartyOnDisconnect)
+			if (USocialSettings::ShouldLeavePartyOnDisconnect())
 			{
 				PersistentParty->LeaveParty();
 			}
diff --git a/Engine/Plugins/Online/OnlineFramework/Source/Party/Private/SocialSettings.cpp b/Engine/Plugins/Online/OnlineFramework/Source/Party/Private/SocialSettings.cpp
--- a/Engine/Plugins/Online/OnlineFramework/Source/Party/Private/SocialSettings.cpp
+++ b/Engine/Plugins/Online/OnlineFramework/Source/Party/Private/SocialSettings.cpp
@@ -43,6 +43,12 @@ bool USocialSettings::MustSendPrimaryInvites()
 	return SettingsCDO.bMustSendPrimaryInvites;
 }
 
+bool USocialSettings::ShouldLeavePartyOnDisconnect()
+{
+	const USocialSettings& SettingsCDO = *GetDefault<USocialSettings>();
+	return SettingsCDO.bLeavePartyOnDisconnect;
+}
+
 int32 USocialSettings::GetDefaultMaxPartySize()
 {
 #if !UE_BUILD_SHIPPING
diff --git a/Engine/Plugins/Online/OnlineFramework/Source/Party/Public/SocialManager.h b/Engine/Plugins/Online/OnlineFramework/Source/Party/Public/SocialManager.h
--- a/Engine/Plugins/Online/OnlineFramework/Source/Party/Public/SocialManager.h
+++ b/Engine/Plugins/Online/OnlineFramework/Source/Party/Public/SocialManager.h
@@ -253,10 +253,6 @@ private:
 	 */
 	bool bCanCreatePartyObjects = false;
 
-	/** Should we leave a party when it enters the disconnected state? */
-	UPROPERTY(config)
-	bool bLeavePartyOnDisconnect = true;
-
 	TSharedPtr<FPartyPlatformSessionManager> PartySessionManager;
 
 	TMap<FOnlinePartyTypeId, USocialParty*> JoinedPartiesByTypeId;
diff --git a/Engine/Plugins/Online/OnlineFramework/Source/Party/Public/SocialSettings.h b/Engine/Plugins/Online/OnlineFramework/Source/Party/Public/SocialSettings.h
--- a/Engine/Plugins/Online/OnlineFramework/Source/Party/Public/SocialSettings.h
+++ b/Engine/Plugins/Online/OnlineFramework/Source/Party/Public/SocialSettings.h
@@ -21,9 +21,10 @@ public:
 	USocialSettings();
 
 	static FString GetUniqueIdEnvironmentPrefix(ESocialSubsystem SubsystemType);
+	static int32 GetDefaultMaxPartySize();
 	static bool ShouldPreferPlatformInvites();
 	static bool MustSendPrimaryInvites();
-	static int32 GetDefaultMaxPartySize();
+	static bool ShouldLeavePartyOnDisconnect();
 	static float GetUserListAutoUpdateRate();
 
 private:
@@ -46,6 +47,10 @@ private:
 	UPROPERTY(config)
 	bool bMustSendPrimaryInvites = false;
 
+	/** Should we leave a party when it enters the disconnected state? */
+	UPROPERTY(config)
+	bool bLeavePartyOnDisconnect = true;
+
 	UPROPERTY(config)
 	float UserListAutoUpdateRate = 7.5f;
 };
\ No newline at end of file
