Fix shutdown crash by replacing checkf with if check

diff --git a/Engine/Plugins/Online/OnlineFramework/Source/Party/Private/SocialManager.cpp b/Engine/Plugins/Online/OnlineFramework/Source/Party/Private/SocialManager.cpp
--- a/Engine/Plugins/Online/OnlineFramework/Source/Party/Private/SocialManager.cpp
+++ b/Engine/Plugins/Online/OnlineFramework/Source/Party/Private/SocialManager.cpp
@@ -707,20 +707,25 @@ void USocialManager::JoinPartyInternal(FJoinPartyAttempt& JoinAttempt)
 {
 	IOnlinePartyPtr PartyInterface = Online::GetPartyInterfaceChecked(GetWorld());
 	FUniqueNetIdRepl LocalUserId = GetFirstLocalUserId(ESocialSubsystem::Primary);
-	checkf(LocalUserId.IsValid(), TEXT("USocialManager::JoinPartyInternal: Invalid LocalUserId!"));
-
-	JoinAttempt.ActionTimeTracker.BeginStep(FJoinPartyAttempt::Step_JoinParty);
-
-	if (JoinAttempt.RejoinInfo.IsValid())
+	if (LocalUserId.IsValid())
 	{
-		UE_LOG(LogParty, Verbose, TEXT("Attempting to rejoin party [%s] now."), *JoinAttempt.RejoinInfo->PartyId->ToDebugString());
+		JoinAttempt.ActionTimeTracker.BeginStep(FJoinPartyAttempt::Step_JoinParty);
+
+		if (JoinAttempt.RejoinInfo.IsValid())
+		{
+			UE_LOG(LogParty, Verbose, TEXT("Attempting to rejoin party [%s] now."), *JoinAttempt.RejoinInfo->PartyId->ToDebugString());
 
-		// Rejoin attempts are initiated differently, but the handler/follow-up is identical to a normal join
-		PartyInterface->RejoinParty(*LocalUserId, *JoinAttempt.RejoinInfo->PartyId, IOnlinePartySystem::GetPrimaryPartyTypeId(), JoinAttempt.RejoinInfo->MemberIds, FOnJoinPartyComplete::CreateUObject(this, &USocialManager::HandleJoinPartyComplete, IOnlinePartySystem::GetPrimaryPartyTypeId()));
+			// Rejoin attempts are initiated differently, but the handler/follow-up is identical to a normal join
+			PartyInterface->RejoinParty(*LocalUserId, *JoinAttempt.RejoinInfo->PartyId, IOnlinePartySystem::GetPrimaryPartyTypeId(), JoinAttempt.RejoinInfo->MemberIds, FOnJoinPartyComplete::CreateUObject(this, &USocialManager::HandleJoinPartyComplete, IOnlinePartySystem::GetPrimaryPartyTypeId()));
+		}
+		else
+		{
+			PartyInterface->JoinParty(*LocalUserId, *JoinAttempt.JoinInfo, FOnJoinPartyComplete::CreateUObject(this, &USocialManager::HandleJoinPartyComplete, JoinAttempt.JoinInfo->GetPartyTypeId()));
+		}
 	}
 	else
 	{
-		PartyInterface->JoinParty(*LocalUserId, *JoinAttempt.JoinInfo, FOnJoinPartyComplete::CreateUObject(this, &USocialManager::HandleJoinPartyComplete, JoinAttempt.JoinInfo->GetPartyTypeId()));
+		UE_LOG(LogParty, Error, TEXT("USocialManager::JoinPartyInternal Invalid LocalUserId=[%s] on primary subsystem."), *LocalUserId.ToDebugString());
 	}
 }
 
