-- Add nickname min/max limit and validation filter to social settings.
  -- And adjust social user's callback to contain error message.

diff --git a/Engine/Plugins/Online/OnlineFramework/Source/Party/Private/SocialSettings.cpp b/Engine/Plugins/Online/OnlineFramework/Source/Party/Private/SocialSettings.cpp
--- a/Engine/Plugins/Online/OnlineFramework/Source/Party/Private/SocialSettings.cpp
+++ b/Engine/Plugins/Online/OnlineFramework/Source/Party/Private/SocialSettings.cpp
@@ -3,6 +3,7 @@
 #include "SocialSettings.h"
 #include "SocialManager.h"
 #include "Misc/CommandLine.h"
+#include "Internationalization/Regex.h"
 
 #if !UE_BUILD_SHIPPING
 int32 MaxPartySizeOverride = INDEX_NONE;
@@ -73,3 +74,27 @@ float USocialSettings::GetUserListAutoUpdateRate()
 	const USocialSettings& SettingsCDO = *GetDefault<USocialSettings>();
 	return SettingsCDO.UserListAutoUpdateRate;
 }
+
+int32 USocialSettings::GetMinNicknameLength()
+{
+	const USocialSettings& SettingsCDO = *GetDefault<USocialSettings>();
+	return SettingsCDO.MinNicknameLength;
+}
+
+int32 USocialSettings::GetMaxNicknameLength()
+{
+	const USocialSettings& SettingsCDO = *GetDefault<USocialSettings>();
+	return SettingsCDO.MaxNicknameLength;
+}
+
+bool USocialSettings::ValidateNickname(const FString& InNickname)
+{
+	const USocialSettings& SettingsCDO = *GetDefault<USocialSettings>();
+	FRegexPattern MatchPattern(SettingsCDO.NicknameRegexPattern);
+	FRegexMatcher Matcher(MatchPattern, InNickname);
+	if (Matcher.FindNext())
+	{
+		return Matcher.GetMatchEnding() == InNickname.Len();
+	}
+	return false;
+}
diff --git a/Engine/Plugins/Online/OnlineFramework/Source/Party/Private/User/SocialUser.cpp b/Engine/Plugins/Online/OnlineFramework/Source/Party/Private/User/SocialUser.cpp
--- a/Engine/Plugins/Online/OnlineFramework/Source/Party/Private/User/SocialUser.cpp
+++ b/Engine/Plugins/Online/OnlineFramework/Source/Party/Private/User/SocialUser.cpp
@@ -435,11 +435,17 @@ FString USocialUser::GetNickname() const
 
 void USocialUser::SetNickname(const FString& InNickname)
 {
-	// @note StephanJ: Nickname for now only applies to primary/MCP subsystem.
 	IOnlineFriendsPtr FriendsInterface = GetOwningToolkit().GetSocialOss(ESocialSubsystem::Primary)->GetFriendsInterface();
 	check(FriendsInterface.IsValid());
 
-	FriendsInterface->SetFriendAlias(GetOwningToolkit().GetLocalUserNum(), *GetUserId(ESocialSubsystem::Primary), EFriendsLists::ToString(EFriendsLists::Default), InNickname, FOnSetFriendAliasComplete::CreateUObject(const_cast<USocialUser*>(this), &USocialUser::HandleSetNicknameComplete));
+	if (!InNickname.IsEmpty())
+	{
+		FriendsInterface->SetFriendAlias(GetOwningToolkit().GetLocalUserNum(), *GetUserId(ESocialSubsystem::Primary), EFriendsLists::ToString(EFriendsLists::Default), InNickname, FOnSetFriendAliasComplete::CreateUObject(const_cast<USocialUser*>(this), &USocialUser::HandleSetNicknameComplete));
+	}
+	else if (!GetNickname().IsEmpty())
+	{
+		FriendsInterface->DeleteFriendAlias(GetOwningToolkit().GetLocalUserNum(), *GetUserId(ESocialSubsystem::Primary), EFriendsLists::ToString(EFriendsLists::Default), FOnDeleteFriendAliasComplete::CreateUObject(const_cast<USocialUser*>(this), &USocialUser::HandleSetNicknameComplete));
+	}
 }
 
 EInviteStatus::Type USocialUser::GetFriendInviteStatus(ESocialSubsystem SubsystemType) const
@@ -1286,12 +1292,9 @@ void USocialUser::HandleSetNicknameComplete(int32 LocalUserNum, const FUniqueNet
 {
 	if (!Error.WasSuccessful())
 	{
-		UE_LOG(LogOnline, Warning, TEXT("Set nickname request failed for user: %s"), *FriendId.ToDebugString());
-	}
-	else
-	{
-		OnNicknameChanged().Broadcast();
+		UE_LOG(LogOnline, Warning, TEXT("Set nickname request failed for user: %s with error message: %s"), *FriendId.ToDebugString(), *Error.GetErrorMessage().ToString());
 	}
+	OnSetNicknameCompleted().Broadcast(Error.GetErrorMessage());
 }
 
 FString USocialUser::SanitizePresenceString(FString InString) const
diff --git a/Engine/Plugins/Online/OnlineFramework/Source/Party/Public/SocialSettings.h b/Engine/Plugins/Online/OnlineFramework/Source/Party/Public/SocialSettings.h
--- a/Engine/Plugins/Online/OnlineFramework/Source/Party/Public/SocialSettings.h
+++ b/Engine/Plugins/Online/OnlineFramework/Source/Party/Public/SocialSettings.h
@@ -26,6 +26,9 @@ public:
 	static bool MustSendPrimaryInvites();
 	static bool ShouldLeavePartyOnDisconnect();
 	static float GetUserListAutoUpdateRate();
+	static int32 GetMinNicknameLength();
+	static int32 GetMaxNicknameLength();
+	static bool ValidateNickname(const FString& InNickname);
 
 private:
 	/**
@@ -53,4 +56,13 @@ private:
 
 	UPROPERTY(config)
 	float UserListAutoUpdateRate = 7.5f;
+
+	UPROPERTY(Config)
+	int32 MinNicknameLength = 3;
+
+	UPROPERTY(Config)
+	int32 MaxNicknameLength = 16;
+
+	UPROPERTY(Config)
+	FString NicknameRegexPattern = TEXT("[\\p{L}\\p{N} \\-_.'\\u0091\\u0092\\u2018\\u2019]+");
 };
\ No newline at end of file
diff --git a/Engine/Plugins/Online/OnlineFramework/Source/Party/Public/User/SocialUser.h b/Engine/Plugins/Online/OnlineFramework/Source/Party/Public/User/SocialUser.h
--- a/Engine/Plugins/Online/OnlineFramework/Source/Party/Public/User/SocialUser.h
+++ b/Engine/Plugins/Online/OnlineFramework/Source/Party/Public/User/SocialUser.h
@@ -105,8 +105,8 @@ public:
 
 	UPartyMember* GetPartyMember(const FOnlinePartyTypeId& PartyTypeId) const;
 
-	DECLARE_EVENT(USocialUser, FOnNicknameChanged);
-	FOnNicknameChanged& OnNicknameChanged() const { return OnNicknameChangedEvent; }
+	DECLARE_EVENT_OneParam(USocialUser, FOnNicknameChanged, const FText&);
+	FOnNicknameChanged& OnSetNicknameCompleted() const { return OnSetNicknameCompletedEvent; }
 
 	DECLARE_EVENT(USocialUser, FPartyInviteResponseEvent);
 	FPartyInviteResponseEvent& OnPartyInviteAccepted() const { return OnPartyInviteAcceptedEvent; }
@@ -198,7 +198,7 @@ private:
 	// Initialization delegates that fire only when a specific user has finishing initializing
 	static TMap<TWeakObjectPtr<USocialUser>, FOnNewSocialUserInitialized> InitEventsByUser;
 
-	mutable FOnNicknameChanged OnNicknameChangedEvent;
+	mutable FOnNicknameChanged OnSetNicknameCompletedEvent;
 	mutable FPartyInviteResponseEvent OnPartyInviteAcceptedEvent;
 	mutable FPartyInviteResponseEvent OnPartyInviteRejectedEvent;
 	mutable FOnUserPresenceChanged OnUserPresenceChangedEvent;
