Use const ULocalPlayer in USocialManager::GetFortSocialToolkit

#rnx

diff --git a/Engine/Plugins/Online/OnlineFramework/Source/Party/Private/SocialManager.cpp b/Engine/Plugins/Online/OnlineFramework/Source/Party/Private/SocialManager.cpp
--- a/Engine/Plugins/Online/OnlineFramework/Source/Party/Private/SocialManager.cpp
+++ b/Engine/Plugins/Online/OnlineFramework/Source/Party/Private/SocialManager.cpp
@@ -244,7 +244,7 @@ void USocialManager::ShutdownSocialManager()
 	MarkPendingKill();
 }
 
-USocialToolkit& USocialManager::GetSocialToolkit(ULocalPlayer& LocalPlayer) const
+USocialToolkit& USocialManager::GetSocialToolkit(const ULocalPlayer& LocalPlayer) const
 {
 	USocialToolkit* FoundToolkit = nullptr;
 	for (USocialToolkit* Toolkit : SocialToolkits)
diff --git a/Engine/Plugins/Online/OnlineFramework/Source/Party/Public/SocialManager.h b/Engine/Plugins/Online/OnlineFramework/Source/Party/Public/SocialManager.h
--- a/Engine/Plugins/Online/OnlineFramework/Source/Party/Public/SocialManager.h
+++ b/Engine/Plugins/Online/OnlineFramework/Source/Party/Public/SocialManager.h
@@ -49,7 +49,7 @@ public:
 	virtual void InitSocialManager();
 	virtual void ShutdownSocialManager();
 
-	USocialToolkit& GetSocialToolkit(ULocalPlayer& LocalPlayer) const;
+	USocialToolkit& GetSocialToolkit(const ULocalPlayer& LocalPlayer) const;
 	USocialToolkit* GetFirstLocalUserToolkit() const;
 	FUniqueNetIdRepl GetFirstLocalUserId(ESocialSubsystem SubsystemType) const;
 	bool IsLocalUser(const FUniqueNetIdRepl& LocalUserId, ESocialSubsystem SubsystemType) const;
