Make SetNickname virtual, and fixes the regex validation for nickname

diff --git a/Engine/Plugins/Online/OnlineFramework/Source/Party/Private/SocialSettings.cpp b/Engine/Plugins/Online/OnlineFramework/Source/Party/Private/SocialSettings.cpp
--- a/Engine/Plugins/Online/OnlineFramework/Source/Party/Private/SocialSettings.cpp
+++ b/Engine/Plugins/Online/OnlineFramework/Source/Party/Private/SocialSettings.cpp
@@ -94,7 +94,7 @@ bool USocialSettings::ValidateNickname(const FString& InNickname)
 	FRegexMatcher Matcher(MatchPattern, InNickname);
 	if (Matcher.FindNext())
 	{
-		return Matcher.GetMatchEnding() == InNickname.Len();
+		return Matcher.GetMatchBeginning() == 0 && Matcher.GetMatchEnding() == InNickname.Len();
 	}
 	return false;
 }
diff --git a/Engine/Plugins/Online/OnlineFramework/Source/Party/Public/User/SocialUser.h b/Engine/Plugins/Online/OnlineFramework/Source/Party/Public/User/SocialUser.h
--- a/Engine/Plugins/Online/OnlineFramework/Source/Party/Public/User/SocialUser.h
+++ b/Engine/Plugins/Online/OnlineFramework/Source/Party/Public/User/SocialUser.h
@@ -47,7 +47,7 @@ public:
 	FString GetDisplayName() const;
 	FString GetDisplayName(ESocialSubsystem SubsystemType) const;
 	FString GetNickname() const;
-	void SetNickname(const FString& InNickName);
+	virtual void SetNickname(const FString& InNickName);
 
 	EInviteStatus::Type GetFriendInviteStatus(ESocialSubsystem SubsystemType) const;
 	bool IsFriend() const;
