Make SetNickname return a bool to indicate whether action was taken (and that we should wait for a response) or not.
And make the handler for SetNicknameComplete virtual.

diff --git a/Engine/Plugins/Online/OnlineFramework/Source/Party/Private/User/SocialUser.cpp b/Engine/Plugins/Online/OnlineFramework/Source/Party/Private/User/SocialUser.cpp
--- a/Engine/Plugins/Online/OnlineFramework/Source/Party/Private/User/SocialUser.cpp
+++ b/Engine/Plugins/Online/OnlineFramework/Source/Party/Private/User/SocialUser.cpp
@@ -433,7 +433,7 @@ FString USocialUser::GetNickname() const
 	return TEXT("");
 }
 
-void USocialUser::SetNickname(const FString& InNickname)
+bool USocialUser::SetNickname(const FString& InNickname)
 {
 	IOnlineFriendsPtr FriendsInterface = GetOwningToolkit().GetSocialOss(ESocialSubsystem::Primary)->GetFriendsInterface();
 	check(FriendsInterface.IsValid());
@@ -441,11 +441,15 @@ void USocialUser::SetNickname(const FString& InNickname)
 	if (!InNickname.IsEmpty())
 	{
 		FriendsInterface->SetFriendAlias(GetOwningToolkit().GetLocalUserNum(), *GetUserId(ESocialSubsystem::Primary), EFriendsLists::ToString(EFriendsLists::Default), InNickname, FOnSetFriendAliasComplete::CreateUObject(const_cast<USocialUser*>(this), &USocialUser::HandleSetNicknameComplete));
+		return true;
 	}
 	else if (!GetNickname().IsEmpty())
 	{
 		FriendsInterface->DeleteFriendAlias(GetOwningToolkit().GetLocalUserNum(), *GetUserId(ESocialSubsystem::Primary), EFriendsLists::ToString(EFriendsLists::Default), FOnDeleteFriendAliasComplete::CreateUObject(const_cast<USocialUser*>(this), &USocialUser::HandleSetNicknameComplete));
+		return true;
 	}
+	OnSetNicknameCompleted().Broadcast(FText::GetEmpty());
+	return false;
 }
 
 EInviteStatus::Type USocialUser::GetFriendInviteStatus(ESocialSubsystem SubsystemType) const
diff --git a/Engine/Plugins/Online/OnlineFramework/Source/Party/Public/User/SocialUser.h b/Engine/Plugins/Online/OnlineFramework/Source/Party/Public/User/SocialUser.h
--- a/Engine/Plugins/Online/OnlineFramework/Source/Party/Public/User/SocialUser.h
+++ b/Engine/Plugins/Online/OnlineFramework/Source/Party/Public/User/SocialUser.h
@@ -46,8 +46,8 @@ public:
 	FUniqueNetIdRepl GetUserId(ESocialSubsystem SubsystemType) const;
 	FString GetDisplayName() const;
 	FString GetDisplayName(ESocialSubsystem SubsystemType) const;
-	FString GetNickname() const;
-	virtual void SetNickname(const FString& InNickName);
+	virtual FString GetNickname() const;
+	virtual bool SetNickname(const FString& InNickName);
 
 	EInviteStatus::Type GetFriendInviteStatus(ESocialSubsystem SubsystemType) const;
 	bool IsFriend() const;
@@ -153,15 +153,14 @@ protected:
 	virtual void OnPresenceChangedInternal(ESocialSubsystem SubsystemType);
 	virtual void OnPartyInviteAcceptedInternal(const FOnlinePartyTypeId& PartyTypeId) const;
 	virtual void OnPartyInviteRejectedInternal(const FOnlinePartyTypeId& PartyTypeId) const;
+	virtual void HandleSetNicknameComplete(int32 LocalUserNum, const FUniqueNetId& FriendId, const FString& ListName, const FOnlineError& Error);
 	virtual void SetSubsystemId(ESocialSubsystem SubsystemType, const FUniqueNetIdRepl& SubsystemId);
 	int32 NumPendingQueries = 0;
 
 	void TryBroadcastInitializationComplete();
 private:
-	
 	void SetUserInfo(ESocialSubsystem SubsystemType, const TSharedRef<FOnlineUser>& UserInfo);
 	void HandleQueryUserInfoComplete(ESocialSubsystem SubsystemType, bool bWasSuccessful, const TSharedPtr<FOnlineUser>& UserInfo);
-	void HandleSetNicknameComplete(int32 LocalUserNum, const FUniqueNetId& FriendId, const FString& ListName, const FOnlineError& Error);
 
 	virtual FString SanitizePresenceString(FString InString) const;
 	
