Reset FSocialUserList::UpdateTickerHandle in FSocialUserList::SetAllowAutoUpdate. The ticker is not added if UpdateTickerHandle.IsValid().

Further, changed USocialSettings::UserListAutoUpdateRate from 7.5 seconds to 0.5 second.

#tests Compile DebugGame Editor Win64 / Development Editor Win64 / Ran it in PIE with two clients and verified fix.

diff --git a/Engine/Plugins/Online/OnlineFramework/Source/Party/Private/User/SocialUserList.cpp b/Engine/Plugins/Online/OnlineFramework/Source/Party/Private/User/SocialUserList.cpp
--- a/Engine/Plugins/Online/OnlineFramework/Source/Party/Private/User/SocialUserList.cpp
+++ b/Engine/Plugins/Online/OnlineFramework/Source/Party/Private/User/SocialUserList.cpp
@@ -99,6 +99,7 @@ void FSocialUserList::SetAllowAutoUpdate(bool bIsEnabled)
 	if (!bIsEnabled && UpdateTickerHandle.IsValid())
 	{
 		FTicker::GetCoreTicker().RemoveTicker(UpdateTickerHandle);
+		UpdateTickerHandle.Reset();
 	}
 	else if (bIsEnabled && !UpdateTickerHandle.IsValid())
 	{
diff --git a/Engine/Plugins/Online/OnlineFramework/Source/Party/Public/SocialSettings.h b/Engine/Plugins/Online/OnlineFramework/Source/Party/Public/SocialSettings.h
--- a/Engine/Plugins/Online/OnlineFramework/Source/Party/Public/SocialSettings.h
+++ b/Engine/Plugins/Online/OnlineFramework/Source/Party/Public/SocialSettings.h
@@ -55,7 +55,7 @@ private:
 	bool bLeavePartyOnDisconnect = true;
 
 	UPROPERTY(config)
-	float UserListAutoUpdateRate = 7.5f;
+	float UserListAutoUpdateRate = 0.5f;
 
 	UPROPERTY(Config)
 	int32 MinNicknameLength = 3;
