Have hotfix manager inform the HTTP module when any HTTP section is updated (e.g. [HTTP.Curl])

diff --git a/Engine/Plugins/Online/OnlineFramework/Source/Hotfix/Private/OnlineHotfixManager.cpp b/Engine/Plugins/Online/OnlineFramework/Source/Hotfix/Private/OnlineHotfixManager.cpp
--- a/Engine/Plugins/Online/OnlineFramework/Source/Hotfix/Private/OnlineHotfixManager.cpp
+++ b/Engine/Plugins/Online/OnlineFramework/Source/Hotfix/Private/OnlineHotfixManager.cpp
@@ -911,7 +911,7 @@ bool UOnlineHotfixManager::HotfixIniFile(const FString& FileName, const FString&
 				{
 					const TCHAR* LogConfigSection = TEXT("[Core.Log]");
 					const TCHAR* ConsoleVariableSection = TEXT("[ConsoleVariables]");
-					const TCHAR* HttpSection = TEXT("[HTTP]");
+					const TCHAR* HttpSection = TEXT("[HTTP"); // note "]" omitted on purpose since we want a partial match
 					const TCHAR* OnlineSubSectionKey = TEXT("[OnlineSubsystem"); // note "]" omitted on purpose since we want a partial match
 					if (!bUpdateLogSuppression && FCString::Strnicmp(*IniData + StartIndex, LogConfigSection, FCString::Strlen(LogConfigSection)) == 0)
 					{
