Make RejoinCheck async functions callable with in-line default-constructed delegates

diff --git a/Engine/Plugins/Online/OnlineFramework/Source/Rejoin/Private/RejoinCheck.cpp b/Engine/Plugins/Online/OnlineFramework/Source/Rejoin/Private/RejoinCheck.cpp
--- a/Engine/Plugins/Online/OnlineFramework/Source/Rejoin/Private/RejoinCheck.cpp
+++ b/Engine/Plugins/Online/OnlineFramework/Source/Rejoin/Private/RejoinCheck.cpp
@@ -204,7 +204,7 @@ void URejoinCheck::RejoinCheckTimer()
 	}
 }
 
-void URejoinCheck::RejoinLastSession(FOnRejoinLastSessionComplete& InCompletionDelegate)
+void URejoinCheck::RejoinLastSession(const FOnRejoinLastSessionComplete& InCompletionDelegate)
 {
 	if (LastKnownStatus == ERejoinStatus::UpdatingStatus)
 	{
diff --git a/Engine/Plugins/Online/OnlineFramework/Source/Rejoin/Public/RejoinCheck.h b/Engine/Plugins/Online/OnlineFramework/Source/Rejoin/Public/RejoinCheck.h
--- a/Engine/Plugins/Online/OnlineFramework/Source/Rejoin/Public/RejoinCheck.h
+++ b/Engine/Plugins/Online/OnlineFramework/Source/Rejoin/Public/RejoinCheck.h
@@ -97,7 +97,7 @@ public:
 	 * 
 	 * @param InCompletionDelegate delegate called after the rejoin attempt is complete
 	 */
-	void RejoinLastSession(FOnRejoinLastSessionComplete& InCompletionDelegate);
+	void RejoinLastSession(const FOnRejoinLastSessionComplete& InCompletionDelegate);
 
 	/**
 	 * Manually set the status of rejoins.  Used when entering/leaving a map as a hint for future check requirements
