Check dirtyuser weakptr validity before dereferencing it, and add additional logging to HandleBlockPlayerComplete.

diff --git a/Engine/Plugins/Online/OnlineFramework/Source/Party/Private/SocialToolkit.cpp b/Engine/Plugins/Online/OnlineFramework/Source/Party/Private/SocialToolkit.cpp
--- a/Engine/Plugins/Online/OnlineFramework/Source/Party/Private/SocialToolkit.cpp
+++ b/Engine/Plugins/Online/OnlineFramework/Source/Party/Private/SocialToolkit.cpp
@@ -940,6 +940,7 @@ void USocialToolkit::HandleBlockPlayerComplete(int32 LocalUserNum, bool bWasSucc
 
 					if (BlockedPlayerInfoPtr)
 					{
+						UE_LOG(LogParty, Log, TEXT("%s Block player %s complete, Establishing Oss info..."), ANSI_TO_TCHAR(__FUNCTION__), *BlockedUserId.ToDebugString());
 						User.EstablishOssInfo(*BlockedPlayerInfoPtr, SubsystemType);
 						OnUserBlocked().Broadcast(User, SubsystemType, true);
 					}
diff --git a/Engine/Plugins/Online/OnlineFramework/Source/Party/Private/User/SocialUserList.cpp b/Engine/Plugins/Online/OnlineFramework/Source/Party/Private/User/SocialUserList.cpp
--- a/Engine/Plugins/Online/OnlineFramework/Source/Party/Private/User/SocialUserList.cpp
+++ b/Engine/Plugins/Online/OnlineFramework/Source/Party/Private/User/SocialUserList.cpp
@@ -491,17 +491,20 @@ void FSocialUserList::UpdateListInternal()
 	// Re-evaluate whether each user with dirtied presence is still fit for the list
 	for (TWeakObjectPtr<USocialUser> DirtyUser : UsersWithDirtyPresence)
 	{
-		const bool bContainsUser = Users.Contains(DirtyUser);
-		const bool bPendingAdd = PendingAdds.Contains(DirtyUser);
-		const bool bPendingRemove = PendingRemovals.Contains(DirtyUser);
-
-		if (bPendingRemove || (!bContainsUser && !bPendingAdd))
-		{
-			TryAddUserFast(*DirtyUser);
-		}
-		else if (bPendingAdd || (bContainsUser && !bPendingRemove))
+		if (DirtyUser.IsValid())
 		{
-			TryRemoveUser(*DirtyUser);
+			const bool bContainsUser = Users.Contains(DirtyUser);
+			const bool bPendingAdd = PendingAdds.Contains(DirtyUser);
+			const bool bPendingRemove = PendingRemovals.Contains(DirtyUser);
+
+			if (bPendingRemove || (!bContainsUser && !bPendingAdd))
+			{
+				TryAddUserFast(*DirtyUser);
+			}
+			else if (bPendingAdd || (bContainsUser && !bPendingRemove))
+			{
+				TryRemoveUser(*DirtyUser);
+			}
 		}
 	}
 	UsersWithDirtyPresence.Reset();
