Add IsFriend() check for SetNickname call

diff --git a/Engine/Plugins/Online/OnlineFramework/Source/Party/Private/User/SocialUser.cpp b/Engine/Plugins/Online/OnlineFramework/Source/Party/Private/User/SocialUser.cpp
--- a/Engine/Plugins/Online/OnlineFramework/Source/Party/Private/User/SocialUser.cpp
+++ b/Engine/Plugins/Online/OnlineFramework/Source/Party/Private/User/SocialUser.cpp
@@ -435,18 +435,21 @@ FString USocialUser::GetNickname() const
 
 bool USocialUser::SetNickname(const FString& InNickname)
 {
-	IOnlineFriendsPtr FriendsInterface = GetOwningToolkit().GetSocialOss(ESocialSubsystem::Primary)->GetFriendsInterface();
-	check(FriendsInterface.IsValid());
-
-	if (!InNickname.IsEmpty())
-	{
-		FriendsInterface->SetFriendAlias(GetOwningToolkit().GetLocalUserNum(), *GetUserId(ESocialSubsystem::Primary), EFriendsLists::ToString(EFriendsLists::Default), InNickname, FOnSetFriendAliasComplete::CreateUObject(const_cast<USocialUser*>(this), &USocialUser::HandleSetNicknameComplete));
-		return true;
-	}
-	else if (!GetNickname().IsEmpty())
+	if (IsFriend())
 	{
-		FriendsInterface->DeleteFriendAlias(GetOwningToolkit().GetLocalUserNum(), *GetUserId(ESocialSubsystem::Primary), EFriendsLists::ToString(EFriendsLists::Default), FOnDeleteFriendAliasComplete::CreateUObject(const_cast<USocialUser*>(this), &USocialUser::HandleSetNicknameComplete));
-		return true;
+		IOnlineFriendsPtr FriendsInterface = GetOwningToolkit().GetSocialOss(ESocialSubsystem::Primary)->GetFriendsInterface();
+		check(FriendsInterface.IsValid());
+
+		if (!InNickname.IsEmpty())
+		{
+			FriendsInterface->SetFriendAlias(GetOwningToolkit().GetLocalUserNum(), *GetUserId(ESocialSubsystem::Primary), EFriendsLists::ToString(EFriendsLists::Default), InNickname, FOnSetFriendAliasComplete::CreateUObject(const_cast<USocialUser*>(this), &USocialUser::HandleSetNicknameComplete));
+			return true;
+		}
+		else if (!GetNickname().IsEmpty())
+		{
+			FriendsInterface->DeleteFriendAlias(GetOwningToolkit().GetLocalUserNum(), *GetUserId(ESocialSubsystem::Primary), EFriendsLists::ToString(EFriendsLists::Default), FOnDeleteFriendAliasComplete::CreateUObject(const_cast<USocialUser*>(this), &USocialUser::HandleSetNicknameComplete));
+			return true;
+		}
 	}
 	OnSetNicknameCompleted().Broadcast(FText::GetEmpty());
 	return false;
@@ -1310,7 +1313,7 @@ void USocialUser::HandleSetNicknameComplete(int32 LocalUserNum, const FUniqueNet
 {
 	if (!Error.WasSuccessful())
 	{
-		UE_LOG(LogOnline, Warning, TEXT("Set nickname request failed for user: %s with error message: %s"), *FriendId.ToDebugString(), *Error.GetErrorMessage().ToString());
+		UE_LOG(LogOnline, Log, TEXT("Set nickname request failed for user: %s with error message: %s"), *FriendId.ToDebugString(), *Error.GetErrorMessage().ToString());
 	}
 	OnSetNicknameCompleted().Broadcast(Error.GetErrorMessage());
 }
