Fixed hotfixing for [SystemSettings] config section.

diff --git a/Engine/Plugins/Online/OnlineFramework/Source/Hotfix/Private/OnlineHotfixManager.cpp b/Engine/Plugins/Online/OnlineFramework/Source/Hotfix/Private/OnlineHotfixManager.cpp
--- a/Engine/Plugins/Online/OnlineFramework/Source/Hotfix/Private/OnlineHotfixManager.cpp
+++ b/Engine/Plugins/Online/OnlineFramework/Source/Hotfix/Private/OnlineHotfixManager.cpp
@@ -911,6 +911,7 @@ bool UOnlineHotfixManager::HotfixIniFile(const FString& FileName, const FString&
 				{
 					const TCHAR* LogConfigSection = TEXT("[Core.Log]");
 					const TCHAR* ConsoleVariableSection = TEXT("[ConsoleVariables]");
+					const TCHAR* SystemSettingsSection = TEXT("[SystemSettings]");
 					const TCHAR* HttpSection = TEXT("[HTTP"); // note "]" omitted on purpose since we want a partial match
 					const TCHAR* OnlineSubSectionKey = TEXT("[OnlineSubsystem"); // note "]" omitted on purpose since we want a partial match
 					if (!bUpdateLogSuppression && FCString::Strnicmp(*IniData + StartIndex, LogConfigSection, FCString::Strlen(LogConfigSection)) == 0)
@@ -921,6 +922,10 @@ bool UOnlineHotfixManager::HotfixIniFile(const FString& FileName, const FString&
 					{
 						bUpdateConsoleVariables = true;
 					}
+					else if (!bUpdateConsoleVariables && FCString::Strnicmp(*IniData + StartIndex, SystemSettingsSection, FCString::Strlen(SystemSettingsSection)) == 0)
+					{
+						bUpdateConsoleVariables = true;
+					}
 					else if (!bUpdateHttpConfigs &&	FCString::Strnicmp(*IniData + StartIndex, HttpSection, FCString::Strlen(HttpSection)) == 0)
 					{
 						bUpdateHttpConfigs = true;
diff --git a/Engine/Source/Runtime/Core/Private/Misc/ConfigCacheIni.cpp b/Engine/Source/Runtime/Core/Private/Misc/ConfigCacheIni.cpp
--- a/Engine/Source/Runtime/Core/Private/Misc/ConfigCacheIni.cpp
+++ b/Engine/Source/Runtime/Core/Private/Misc/ConfigCacheIni.cpp
@@ -3774,6 +3774,9 @@ void FConfigCacheIni::LoadConsoleVariablesFromINI()
 	// We also apply from Engine.ini [ConsoleVariables] section
 	ApplyCVarSettingsFromIni(TEXT("ConsoleVariables"), *GEngineIni, ECVF_SetBySystemSettingsIni);
 
+	// The [SystemSettings] section also supports hotfixing
+	ApplyCVarSettingsFromIni(TEXT("SystemSettings"), *GEngineIni, ECVF_SetBySystemSettingsIni);
+
 	IConsoleManager::Get().CallAllConsoleVariableSinks();
 }
 
