Add GetSubsystemUserInfo in SocialUser

diff --git a/Engine/Plugins/Online/OnlineFramework/Source/Party/Public/User/SocialUser.h b/Engine/Plugins/Online/OnlineFramework/Source/Party/Public/User/SocialUser.h
--- a/Engine/Plugins/Online/OnlineFramework/Source/Party/Public/User/SocialUser.h
+++ b/Engine/Plugins/Online/OnlineFramework/Source/Party/Public/User/SocialUser.h
@@ -158,13 +158,6 @@ protected:
 	int32 NumPendingQueries = 0;
 
 	void TryBroadcastInitializationComplete();
-private:
-	void SetUserInfo(ESocialSubsystem SubsystemType, const TSharedRef<FOnlineUser>& UserInfo);
-	void HandleQueryUserInfoComplete(ESocialSubsystem SubsystemType, bool bWasSuccessful, const TSharedPtr<FOnlineUser>& UserInfo);
-
-	virtual FString SanitizePresenceString(FString InString) const;
-	
-private:
 
 	struct FSubsystemUserInfo
 	{
@@ -189,6 +182,15 @@ private:
 		TWeakPtr<FOnlineRecentPlayer> RecentPlayerInfo;
 		TWeakPtr<FOnlineBlockedPlayer> BlockedPlayerInfo;
 	};
+	const FSubsystemUserInfo* GetSubsystemUserInfo(ESocialSubsystem Subsystem) const { return SubsystemInfoByType.Find(Subsystem); }
+
+private:
+	void SetUserInfo(ESocialSubsystem SubsystemType, const TSharedRef<FOnlineUser>& UserInfo);
+	void HandleQueryUserInfoComplete(ESocialSubsystem SubsystemType, bool bWasSuccessful, const TSharedPtr<FOnlineUser>& UserInfo);
+
+	virtual FString SanitizePresenceString(FString InString) const;
+
+private:
 	FSubsystemUserInfo& FindOrCreateSubsystemInfo(const FUniqueNetIdRepl& SubsystemId, ESocialSubsystem SubsystemType);
 
 	bool bIsInitialized = false;
