Adding conditions so that a SocialChatManager only acts on the local player that owns it.

diff --git a/Engine/Plugins/Online/OnlineFramework/Source/Party/Public/Chat/SocialChatManager.cpp b/Engine/Plugins/Online/OnlineFramework/Source/Party/Public/Chat/SocialChatManager.cpp
--- a/Engine/Plugins/Online/OnlineFramework/Source/Party/Public/Chat/SocialChatManager.cpp
+++ b/Engine/Plugins/Online/OnlineFramework/Source/Party/Public/Chat/SocialChatManager.cpp
@@ -349,121 +349,151 @@ USocialChatChannel& USocialChatManager::FindOrCreateChannel(const FText& Display
 
 void USocialChatManager::HandleChatRoomCreated(const FUniqueNetId& LocalUserId, const FChatRoomId& RoomId, bool bWasSuccessful, const FString& Error)
 {
-	if (bWasSuccessful)
+	if (IsUniqueIdOfOwner(LocalUserId))
 	{
-		FindOrCreateRoom(RoomId);
-	}
-	else
-	{
-		HandleChatRoomCreatedFailure(LocalUserId, RoomId, Error);
+		if (bWasSuccessful)
+		{
+			FindOrCreateRoom(RoomId);
+		}
+		else
+		{
+			HandleChatRoomCreatedFailure(LocalUserId, RoomId, Error);
+		}
 	}
 }
 
 void USocialChatManager::HandleChatRoomConfigured(const FUniqueNetId& LocalUserId, const FChatRoomId& RoomId, bool bWasSuccessful, const FString& Error)
 {
-	if (bWasSuccessful)
-	{
-		FindOrCreateRoom(RoomId);
-	}
-	else
+	if (IsUniqueIdOfOwner(LocalUserId))
 	{
-		HandleChatRoomConfiguredFailure(LocalUserId, RoomId, Error);
+		if (bWasSuccessful)
+		{
+			FindOrCreateRoom(RoomId);
+		}
+		else
+		{
+			HandleChatRoomConfiguredFailure(LocalUserId, RoomId, Error);
+		}
 	}
 }
 
 void USocialChatManager::HandleChatRoomJoinPublic(const FUniqueNetId& LocalUserId, const FChatRoomId& RoomId, bool bWasSuccessful, const FString& Error)
 {
-	if (bWasSuccessful)
+	if (IsUniqueIdOfOwner(LocalUserId))
 	{
-		FindOrCreateRoom(RoomId);
-	}
-	else
-	{
-		HandleChatRoomJoinPublicFailure(LocalUserId, RoomId, Error);
+		if (bWasSuccessful)
+		{
+			FindOrCreateRoom(RoomId);
+		}
+		else
+		{
+			HandleChatRoomJoinPublicFailure(LocalUserId, RoomId, Error);
+		}
 	}
 }
 
 void USocialChatManager::HandleChatRoomJoinPrivate(const FUniqueNetId& LocalUserId, const FChatRoomId& RoomId, bool bWasSuccessful, const FString& Error)
 {
-	if (bWasSuccessful)
-	{
-		FindOrCreateRoom(RoomId);
-	}
-	else
+	if (IsUniqueIdOfOwner(LocalUserId))
 	{
-		HandleChatRoomJoinPrivateFailure(LocalUserId, RoomId, Error);
+		if (bWasSuccessful)
+		{
+			FindOrCreateRoom(RoomId);
+		}
+		else
+		{
+			HandleChatRoomJoinPrivateFailure(LocalUserId, RoomId, Error);
+		}
 	}
 }
 
 void USocialChatManager::HandleChatRoomExit(const FUniqueNetId& LocalUserId, const FChatRoomId& RoomId, bool bWasSuccessful, const FString& Error)
 {
-	if (bWasSuccessful)
+	if (IsUniqueIdOfOwner(LocalUserId))
 	{
-		USocialChatRoom** Room = ChatRoomsById.Find(RoomId);
-		if (ensure(Room))
+		if (bWasSuccessful)
 		{
-			ChatRoomsById.Remove(RoomId);
-			OnChannelLeftInternal(**Room);
+			USocialChatRoom** Room = ChatRoomsById.Find(RoomId);
+			if (ensure(Room))
+			{
+				ChatRoomsById.Remove(RoomId);
+				OnChannelLeftInternal(**Room);
+			}
+		}
+		else
+		{
+			HandleChatRoomExitFailure(LocalUserId, RoomId, Error);
 		}
-	}
-	else
-	{
-		HandleChatRoomExitFailure(LocalUserId, RoomId, Error);
 	}
 }
 
 void USocialChatManager::HandleChatRoomMemberJoin(const FUniqueNetId& LocalUserId, const FChatRoomId& RoomId, const FUniqueNetId& MemberId)
 {
-	// This is potentially a previously unknown user, so establish them first
-	GetOwningToolkit().QueueUserDependentAction(MemberId.AsShared(),
-		[this, RoomId] (USocialUser& User)
+	if (IsUniqueIdOfOwner(LocalUserId))
+	{
+		// This is potentially a previously unknown user, so establish them first
+		GetOwningToolkit().QueueUserDependentAction(MemberId.AsShared(),
+			[this, RoomId](USocialUser& User)
 		{
 			if (USocialChatRoom** Channel = ChatRoomsById.Find(RoomId))
 			{
 				(*Channel)->NotifyUserJoinedChannel(User);
 			}
 		});
+	}
 }
 
 void USocialChatManager::HandleChatRoomMemberExit(const FUniqueNetId& LocalUserId, const FChatRoomId& RoomId, const FUniqueNetId& MemberId)
 {
-	//@todo SocialChat: Should the channel be tracking users in it? Should be up to a SocialUserList, shouldn't it? 
-	//	Maybe a channel is actually a UserList? Overkill for PMs, but sounds reasonable for a ChatRoom
-	USocialChatRoom* Channel = GetChatRoom(RoomId);
-	USocialUser* SocialUser = GetOwningToolkit().FindUser(MemberId.AsShared());
-	if (ensure(Channel) && SocialUser)
+	if (IsUniqueIdOfOwner(LocalUserId))
 	{
-		Channel->NotifyUserLeftChannel(*SocialUser);
+		//@todo SocialChat: Should the channel be tracking users in it? Should be up to a SocialUserList, shouldn't it? 
+		//	Maybe a channel is actually a UserList? Overkill for PMs, but sounds reasonable for a ChatRoom
+		USocialChatRoom* Channel = GetChatRoom(RoomId);
+		USocialUser* SocialUser = GetOwningToolkit().FindUser(MemberId.AsShared());
+		if (ensure(Channel) && SocialUser)
+		{
+			Channel->NotifyUserLeftChannel(*SocialUser);
+		}
 	}
 }
 
 void USocialChatManager::HandleChatRoomMemberUpdate(const FUniqueNetId& LocalUserId, const FChatRoomId& RoomId, const FUniqueNetId& MemberId)
 {
-	USocialChatRoom* Channel = GetChatRoom(RoomId);
-	USocialUser* SocialUser = GetOwningToolkit().FindUser(MemberId.AsShared());
-	if (ensure(Channel) && ensure(SocialUser))
+	if (IsUniqueIdOfOwner(LocalUserId))
 	{
-		Channel->NotifyChannelUserChanged(*SocialUser);
+		USocialChatRoom* Channel = GetChatRoom(RoomId);
+		USocialUser* SocialUser = GetOwningToolkit().FindUser(MemberId.AsShared());
+		if (ensure(Channel) && ensure(SocialUser))
+		{
+			Channel->NotifyChannelUserChanged(*SocialUser);
+		}
 	}
 }
 
 void USocialChatManager::HandleChatRoomMessageReceived(const FUniqueNetId& LocalUserId, const FChatRoomId& RoomId, const TSharedRef<FChatMessage>& ChatMessage)
 {
-	USocialChatRoom* Channel = GetChatRoom(RoomId);
-	if (ensure(Channel))
+	if (IsUniqueIdOfOwner(LocalUserId))
 	{
-		Channel->NotifyMessageReceived(ChatMessage);
+		USocialChatRoom* Channel = GetChatRoom(RoomId);
+		if (ensure(Channel))
+		{
+			Channel->NotifyMessageReceived(ChatMessage);
+		}
 	}
 }
 
 void USocialChatManager::HandleChatPrivateMessageReceived(const FUniqueNetId& LocalUserId, const TSharedRef<FChatMessage>& ChatMessage)
 {
-	// We can expect that we already know about a user that is sending us a private message and should not have to create one
-	USocialUser* SocialUser = GetOwningToolkit().FindUser(ChatMessage->GetUserId());
-	if (ensure(SocialUser))
+	if (IsUniqueIdOfOwner(LocalUserId))
 	{
-		USocialChatChannel& Channel = FindOrCreateChannel(*SocialUser);
-		Channel.NotifyMessageReceived(ChatMessage);
+		// We can expect that we already know about a user that is sending us a private message and should not have to create one
+		USocialUser* SocialUser = GetOwningToolkit().FindUser(ChatMessage->GetUserId());
+		if (ensure(SocialUser))
+		{
+			USocialChatChannel& Channel = FindOrCreateChannel(*SocialUser);
+			Channel.NotifyMessageReceived(ChatMessage);
+		}
 	}
 }
 
@@ -531,6 +561,13 @@ void USocialChatManager::OnGroupUpdated(const FUniqueNetId& GroupId)
 	printf("");
 }
 
+bool USocialChatManager::IsUniqueIdOfOwner(const FUniqueNetId& LocalUserId) const
+{
+	USocialUser& Owner = GetOwningToolkit().GetLocalUser();
+	const TSharedPtr<const FUniqueNetId> OwnerNetId = Owner.GetUserId(ESocialSubsystem::Primary).GetUniqueNetId();
+	return LocalUserId == *OwnerNetId;
+}
+
 void USocialChatManager::GetGroupChannels(TArray<USocialGroupChannel*>& JoinedChannels) const
 {
 	for (auto& Entry : GroupChannels)
diff --git a/Engine/Plugins/Online/OnlineFramework/Source/Party/Public/Chat/SocialChatManager.h b/Engine/Plugins/Online/OnlineFramework/Source/Party/Public/Chat/SocialChatManager.h
--- a/Engine/Plugins/Online/OnlineFramework/Source/Party/Public/Chat/SocialChatManager.h
+++ b/Engine/Plugins/Online/OnlineFramework/Source/Party/Public/Chat/SocialChatManager.h
@@ -134,6 +134,8 @@ protected:
 
 	void OnGroupUpdated(const FUniqueNetId& GroupId);
 
+	bool IsUniqueIdOfOwner(const FUniqueNetId& LocalUserId) const;
+
 	// END KIAROS GROUP MANAGEMENT
 	//----------------------------------------------------------------------
 
