Remove nickname validation function from social settings.

diff --git a/Engine/Plugins/Online/OnlineFramework/Source/Party/Private/SocialSettings.cpp b/Engine/Plugins/Online/OnlineFramework/Source/Party/Private/SocialSettings.cpp
--- a/Engine/Plugins/Online/OnlineFramework/Source/Party/Private/SocialSettings.cpp
+++ b/Engine/Plugins/Online/OnlineFramework/Source/Party/Private/SocialSettings.cpp
@@ -3,7 +3,6 @@
 #include "SocialSettings.h"
 #include "SocialManager.h"
 #include "Misc/CommandLine.h"
-#include "Internationalization/Regex.h"
 
 #if !UE_BUILD_SHIPPING
 int32 MaxPartySizeOverride = INDEX_NONE;
@@ -85,16 +84,4 @@ int32 USocialSettings::GetMaxNicknameLength()
 {
 	const USocialSettings& SettingsCDO = *GetDefault<USocialSettings>();
 	return SettingsCDO.MaxNicknameLength;
-}
-
-bool USocialSettings::ValidateNickname(const FString& InNickname)
-{
-	const USocialSettings& SettingsCDO = *GetDefault<USocialSettings>();
-	FRegexPattern MatchPattern(SettingsCDO.NicknameRegexPattern);
-	FRegexMatcher Matcher(MatchPattern, InNickname);
-	if (Matcher.FindNext())
-	{
-		return Matcher.GetMatchBeginning() == 0 && Matcher.GetMatchEnding() == InNickname.Len();
-	}
-	return false;
-}
+}
\ No newline at end of file
diff --git a/Engine/Plugins/Online/OnlineFramework/Source/Party/Public/SocialSettings.h b/Engine/Plugins/Online/OnlineFramework/Source/Party/Public/SocialSettings.h
--- a/Engine/Plugins/Online/OnlineFramework/Source/Party/Public/SocialSettings.h
+++ b/Engine/Plugins/Online/OnlineFramework/Source/Party/Public/SocialSettings.h
@@ -28,7 +28,6 @@ public:
 	static float GetUserListAutoUpdateRate();
 	static int32 GetMinNicknameLength();
 	static int32 GetMaxNicknameLength();
-	static bool ValidateNickname(const FString& InNickname);
 
 private:
 	/**
@@ -62,7 +61,4 @@ private:
 
 	UPROPERTY(Config)
 	int32 MaxNicknameLength = 16;
-
-	UPROPERTY(Config)
-	FString NicknameRegexPattern = TEXT("[\\p{L}\\p{N} \\-_.'\\u0091\\u0092\\u2018\\u2019]+");
 };
\ No newline at end of file
