Fix crash if JSON parse fails

diff --git a/Engine/Plugins/Online/OnlineFramework/Source/Party/Private/SocialManager.cpp b/Engine/Plugins/Online/OnlineFramework/Source/Party/Private/SocialManager.cpp
--- a/Engine/Plugins/Online/OnlineFramework/Source/Party/Private/SocialManager.cpp
+++ b/Engine/Plugins/Online/OnlineFramework/Source/Party/Private/SocialManager.cpp
@@ -679,7 +679,7 @@ void USocialManager::RegisterSecondaryPlayer(int32 LocalPlayerNum, const FOnJoin
 			FString JoinInfoStr = PartyInterface->MakeJoinInfoJson(*PrimaryUserId, PersistentParty->GetPartyId());
 			IOnlinePartyJoinInfoConstPtr JoinInfo = PartyInterface->MakeJoinInfoFromJson(JoinInfoStr);
 
-			if (JoinInfo->IsValid())
+			if (JoinInfo && JoinInfo->IsValid())
 			{
 				PartyInterface->JoinParty(*SecondaryUserId, *JoinInfo, Delegate);
 			}
