Moved FriendsInterface::AcceptInvite call from the SocialUser into the SocialToolkit for easier error handling and to better match existing pattern.

#tests Compile DebugGame Editor Win64 / Development Editor Win64 / Shipping Client / Ran it in PIE and verified functionality.

diff --git a/Engine/Plugins/Online/OnlineFramework/Source/Party/Private/SocialToolkit.cpp b/Engine/Plugins/Online/OnlineFramework/Source/Party/Private/SocialToolkit.cpp
--- a/Engine/Plugins/Online/OnlineFramework/Source/Party/Private/SocialToolkit.cpp
+++ b/Engine/Plugins/Online/OnlineFramework/Source/Party/Private/SocialToolkit.cpp
@@ -448,6 +448,18 @@ bool USocialToolkit::TrySendFriendInvite(USocialUser& SocialUser, ESocialSubsyst
 	return false;
 }
 
+bool USocialToolkit::AcceptFriendInvite(const USocialUser& SocialUser, ESocialSubsystem SubsystemType) const
+{
+	if (SocialUser.GetFriendInviteStatus(SubsystemType) == EInviteStatus::PendingInbound)
+	{
+		IOnlineFriendsPtr FriendsInterface = GetSocialOss(SubsystemType)->GetFriendsInterface();
+		check(FriendsInterface.IsValid());
+
+		return FriendsInterface->AcceptInvite(GetLocalUserNum(), *SocialUser.GetUserId(SubsystemType), EFriendsLists::ToString(EFriendsLists::Default), FOnAcceptInviteComplete::CreateUObject(const_cast<USocialToolkit*>(this), &USocialToolkit::HandleAcceptFriendInviteComplete));
+	}
+	return false;
+}
+
 bool USocialToolkit::IsFriendshipRestricted(const USocialUser& SocialUser, ESocialSubsystem SubsystemType) const
 {
 	return false;
@@ -778,11 +790,11 @@ void USocialToolkit::HandleQueryPrimaryUserIdMappingComplete(bool bWasSuccessful
 {
 	if (!IdentifiedUserId.IsValid())
 	{
-		OnFriendInviteComplete(IdentifiedUserId, DisplayName, false, FriendInviteFailureReason::InviteFailReason_NotFound);
+		OnSendFriendInviteComplete(IdentifiedUserId, false, FriendInviteFailureReason::InviteFailReason_NotFound);
 	}
 	else if (RequestingUserId == IdentifiedUserId)
 	{
-		OnFriendInviteComplete(IdentifiedUserId, DisplayName, false, FriendInviteFailureReason::InviteFailReason_AddingSelfFail);
+		OnSendFriendInviteComplete(IdentifiedUserId, false, FriendInviteFailureReason::InviteFailReason_AddingSelfFail);
 	}
 	else
 	{
@@ -791,11 +803,11 @@ void USocialToolkit::HandleQueryPrimaryUserIdMappingComplete(bool bWasSuccessful
 			{
 				if (SocialUser.IsBlocked())
 				{
-					OnFriendInviteComplete(*SocialUser.GetUserId(ESocialSubsystem::Primary), DisplayName, false, FriendInviteFailureReason::InviteFailReason_AddingBlockedFail);
+					OnSendFriendInviteComplete(*SocialUser.GetUserId(ESocialSubsystem::Primary), false, FriendInviteFailureReason::InviteFailReason_AddingBlockedFail);
 				}
 				else if (SocialUser.IsFriend(ESocialSubsystem::Primary))
 				{
-					OnFriendInviteComplete(*SocialUser.GetUserId(ESocialSubsystem::Primary), DisplayName, false, FriendInviteFailureReason::InviteFailReason_AlreadyFriends);
+					OnSendFriendInviteComplete(*SocialUser.GetUserId(ESocialSubsystem::Primary), false, FriendInviteFailureReason::InviteFailReason_AlreadyFriends);
 				}
 				else
 				{
@@ -863,10 +875,10 @@ void USocialToolkit::HandleFriendInviteSent(int32 LocalUserNum, bool bWasSuccess
 	if (bWasSuccessful)
 	{
 		QueueUserDependentActionInternal(InvitedUserId, SubsystemType,
-			[this, SubsystemType, ListName] (USocialUser& SocialUser)
+			[this, SubsystemType, ListName, LocalUserNum] (USocialUser& SocialUser)
 			{
 				IOnlineFriendsPtr FriendsInterface = Online::GetFriendsInterfaceChecked(GetWorld(), USocialManager::GetSocialOssName(SubsystemType));
-				if (TSharedPtr<FOnlineFriend> OssFriend = FriendsInterface->GetFriend(GetLocalUserNum(), *SocialUser.GetUserId(SubsystemType), ListName))
+				if (TSharedPtr<FOnlineFriend> OssFriend = FriendsInterface->GetFriend(LocalUserNum, *SocialUser.GetUserId(SubsystemType), ListName))
 				{
 					SocialUser.EstablishOssInfo(OssFriend.ToSharedRef(), SubsystemType);
 					if (SocialUser.GetFriendInviteStatus(SubsystemType) == EInviteStatus::PendingOutbound)
@@ -876,7 +888,7 @@ void USocialToolkit::HandleFriendInviteSent(int32 LocalUserNum, bool bWasSuccess
 				}
 			});
 	}
-	OnFriendInviteComplete(InvitedUserId, DisplayName, bWasSuccessful, ErrorStr);
+	OnSendFriendInviteComplete(InvitedUserId, bWasSuccessful, ErrorStr);
 }
 
 void USocialToolkit::HandleFriendRemoved(const FUniqueNetId& LocalUserId, const FUniqueNetId& FriendId, ESocialSubsystem SubsystemType)
@@ -903,6 +915,11 @@ void USocialToolkit::HandleDeleteFriendComplete(int32 InLocalUserNum, bool bWasS
 	}
 }
 
+void USocialToolkit::HandleAcceptFriendInviteComplete(int32 LocalUserNum, bool bWasSuccessful, const FUniqueNetId& InviterUserId, const FString& ListName, const FString& ErrorStr)
+{
+	OnAcceptFriendInviteComplete(InviterUserId, bWasSuccessful, ErrorStr);
+}
+
 void USocialToolkit::HandlePartyInviteReceived(const FUniqueNetId& LocalUserId, const FOnlinePartyId& PartyId, const FUniqueNetId& SenderId)
 {
 	if (LocalUserId == *GetLocalUserNetId(ESocialSubsystem::Primary))
diff --git a/Engine/Plugins/Online/OnlineFramework/Source/Party/Private/User/SocialUser.cpp b/Engine/Plugins/Online/OnlineFramework/Source/Party/Private/User/SocialUser.cpp
--- a/Engine/Plugins/Online/OnlineFramework/Source/Party/Private/User/SocialUser.cpp
+++ b/Engine/Plugins/Online/OnlineFramework/Source/Party/Private/User/SocialUser.cpp
@@ -967,14 +967,7 @@ bool USocialUser::SendFriendInvite(ESocialSubsystem SubsystemType)
 
 bool USocialUser::AcceptFriendInvite(ESocialSubsystem SocialSubsystem) const
 {
-	if (GetFriendInviteStatus(SocialSubsystem) == EInviteStatus::PendingInbound)
-	{
-		IOnlineFriendsPtr FriendsInterface = GetOwningToolkit().GetSocialOss(SocialSubsystem)->GetFriendsInterface();
-		check(FriendsInterface.IsValid());
-
-		return FriendsInterface->AcceptInvite(GetOwningToolkit().GetLocalUserNum(), *GetUserId(SocialSubsystem), EFriendsLists::ToString(EFriendsLists::Default));
-	}
-	return false;
+	return GetOwningToolkit().AcceptFriendInvite(*this, SocialSubsystem);
 }
 
 bool USocialUser::RejectFriendInvite(ESocialSubsystem SocialSubsystem) const
diff --git a/Engine/Plugins/Online/OnlineFramework/Source/Party/Public/SocialToolkit.h b/Engine/Plugins/Online/OnlineFramework/Source/Party/Public/SocialToolkit.h
--- a/Engine/Plugins/Online/OnlineFramework/Source/Party/Public/SocialToolkit.h
+++ b/Engine/Plugins/Online/OnlineFramework/Source/Party/Public/SocialToolkit.h
@@ -116,6 +116,8 @@ PARTY_SCOPE:
 
 	bool TrySendFriendInvite(USocialUser& SocialUser, ESocialSubsystem SubsystemType) const;
 
+	bool AcceptFriendInvite(const USocialUser& SocialUser, ESocialSubsystem SubsystemType) const;
+
 	void HandleUserInvalidated(USocialUser* InvalidUser);
 
 #if PLATFORM_PS4
@@ -128,7 +130,8 @@ protected:
 	virtual void OnOwnerLoggedIn();
 	virtual void OnOwnerLoggedOut();
 	
-	virtual void OnFriendInviteComplete(const FUniqueNetId& InvitedUserId, const FString& InvitedUserName, bool bInviteSuceeded, const FString& InviteError) {}
+	virtual void OnSendFriendInviteComplete(const FUniqueNetId& InvitedUserId, bool bInviteSuceeded, const FString& InviteError) {}
+	virtual void OnAcceptFriendInviteComplete(const FUniqueNetId& InviterUserId, bool bInviteSucceeded, const FString& InviteError) {}
 
 	void QueryFriendsLists();
 	void QueryBlockedPlayers();
@@ -198,6 +201,7 @@ private:	// Handlers
 	void HandleFriendRemoved(const FUniqueNetId& LocalUserId, const FUniqueNetId& FormerFriendId, ESocialSubsystem SubsystemType);
 
 	void HandleDeleteFriendComplete(int32 LocalPlayer, bool bWasSuccessful, const FUniqueNetId& FormerFriendId, const FString& ListName, const FString& ErrorStr, ESocialSubsystem SubsystemType);
+	void HandleAcceptFriendInviteComplete(int32 LocalUserNum, bool bWasSuccessful, const FUniqueNetId& InviterUserId, const FString& ListName, const FString& ErrorStr);
 
 	void HandlePartyInviteReceived(const FUniqueNetId& LocalUserId, const FOnlinePartyId& PartyId, const FUniqueNetId& SenderId);
 
