Restore 'InviteeName' to 'OnSendFriendInviteComplete' to provide useful messaging when the user tries to send a friend invitation to a non-existent Epic Friend.

#tests Compile DebugGame Editor Win64 / Development Editor Win64 / Shipping Client / Ran it in PIE and verified fix.

diff --git a/Engine/Plugins/Online/OnlineFramework/Source/Party/Private/SocialToolkit.cpp b/Engine/Plugins/Online/OnlineFramework/Source/Party/Private/SocialToolkit.cpp
--- a/Engine/Plugins/Online/OnlineFramework/Source/Party/Private/SocialToolkit.cpp
+++ b/Engine/Plugins/Online/OnlineFramework/Source/Party/Private/SocialToolkit.cpp
@@ -790,11 +790,11 @@ void USocialToolkit::HandleQueryPrimaryUserIdMappingComplete(bool bWasSuccessful
 {
 	if (!IdentifiedUserId.IsValid())
 	{
-		OnSendFriendInviteComplete(IdentifiedUserId, false, FriendInviteFailureReason::InviteFailReason_NotFound);
+		OnSendFriendInviteComplete(IdentifiedUserId, DisplayName, false, FriendInviteFailureReason::InviteFailReason_NotFound);
 	}
 	else if (RequestingUserId == IdentifiedUserId)
 	{
-		OnSendFriendInviteComplete(IdentifiedUserId, false, FriendInviteFailureReason::InviteFailReason_AddingSelfFail);
+		OnSendFriendInviteComplete(IdentifiedUserId, DisplayName, false, FriendInviteFailureReason::InviteFailReason_AddingSelfFail);
 	}
 	else
 	{
@@ -803,11 +803,11 @@ void USocialToolkit::HandleQueryPrimaryUserIdMappingComplete(bool bWasSuccessful
 			{
 				if (SocialUser.IsBlocked())
 				{
-					OnSendFriendInviteComplete(*SocialUser.GetUserId(ESocialSubsystem::Primary), false, FriendInviteFailureReason::InviteFailReason_AddingBlockedFail);
+					OnSendFriendInviteComplete(*SocialUser.GetUserId(ESocialSubsystem::Primary), DisplayName, false, FriendInviteFailureReason::InviteFailReason_AddingBlockedFail);
 				}
 				else if (SocialUser.IsFriend(ESocialSubsystem::Primary))
 				{
-					OnSendFriendInviteComplete(*SocialUser.GetUserId(ESocialSubsystem::Primary), false, FriendInviteFailureReason::InviteFailReason_AlreadyFriends);
+					OnSendFriendInviteComplete(*SocialUser.GetUserId(ESocialSubsystem::Primary), DisplayName, false, FriendInviteFailureReason::InviteFailReason_AlreadyFriends);
 				}
 				else
 				{
@@ -888,7 +888,7 @@ void USocialToolkit::HandleFriendInviteSent(int32 LocalUserNum, bool bWasSuccess
 				}
 			});
 	}
-	OnSendFriendInviteComplete(InvitedUserId, bWasSuccessful, ErrorStr);
+	OnSendFriendInviteComplete(InvitedUserId, DisplayName, bWasSuccessful, ErrorStr);
 }
 
 void USocialToolkit::HandleFriendRemoved(const FUniqueNetId& LocalUserId, const FUniqueNetId& FriendId, ESocialSubsystem SubsystemType)
diff --git a/Engine/Plugins/Online/OnlineFramework/Source/Party/Public/SocialToolkit.h b/Engine/Plugins/Online/OnlineFramework/Source/Party/Public/SocialToolkit.h
--- a/Engine/Plugins/Online/OnlineFramework/Source/Party/Public/SocialToolkit.h
+++ b/Engine/Plugins/Online/OnlineFramework/Source/Party/Public/SocialToolkit.h
@@ -130,7 +130,7 @@ protected:
 	virtual void OnOwnerLoggedIn();
 	virtual void OnOwnerLoggedOut();
 	
-	virtual void OnSendFriendInviteComplete(const FUniqueNetId& InvitedUserId, bool bInviteSuceeded, const FString& InviteError) {}
+	virtual void OnSendFriendInviteComplete(const FUniqueNetId& InvitedUserId, const FString& InviteeName, bool bInviteSuceeded, const FString& InviteError) {}
 	virtual void OnAcceptFriendInviteComplete(const FUniqueNetId& InviterUserId, bool bInviteSucceeded, const FString& InviteError) {}
 
 	void QueryFriendsLists();
