Covering edge case for USocialChatManager fix

diff --git a/Engine/Plugins/Online/OnlineFramework/Source/Party/Public/Chat/SocialChatManager.cpp b/Engine/Plugins/Online/OnlineFramework/Source/Party/Public/Chat/SocialChatManager.cpp
--- a/Engine/Plugins/Online/OnlineFramework/Source/Party/Public/Chat/SocialChatManager.cpp
+++ b/Engine/Plugins/Online/OnlineFramework/Source/Party/Public/Chat/SocialChatManager.cpp
@@ -564,8 +564,9 @@ void USocialChatManager::OnGroupUpdated(const FUniqueNetId& GroupId)
 bool USocialChatManager::IsUniqueIdOfOwner(const FUniqueNetId& LocalUserId) const
 {
 	USocialUser& Owner = GetOwningToolkit().GetLocalUser();
-	const TSharedPtr<const FUniqueNetId> OwnerNetId = Owner.GetUserId(ESocialSubsystem::Primary).GetUniqueNetId();
-	return LocalUserId == *OwnerNetId;
+	FUniqueNetIdWrapper OwnerNetIdWrapper = FUniqueNetIdWrapper(Owner.GetUserId(ESocialSubsystem::Primary).GetUniqueNetId());
+	FUniqueNetIdWrapper LocalUserIdWrapper = FUniqueNetIdWrapper(LocalUserId);
+	return LocalUserIdWrapper == OwnerNetIdWrapper;
 }
 
 void USocialChatManager::GetGroupChannels(TArray<USocialGroupChannel*>& JoinedChannels) const
