Fix wrongly set social user on party member when we have multiple local players

diff --git a/Engine/Plugins/Online/OnlineFramework/Source/Party/Private/Party/PartyMember.cpp b/Engine/Plugins/Online/OnlineFramework/Source/Party/Private/Party/PartyMember.cpp
--- a/Engine/Plugins/Online/OnlineFramework/Source/Party/Private/Party/PartyMember.cpp
+++ b/Engine/Plugins/Online/OnlineFramework/Source/Party/Private/Party/PartyMember.cpp
@@ -70,7 +70,7 @@ void UPartyMember::InitializePartyMember(const FOnlinePartyMemberConstRef& InOss
 	{
 		OssPartyMember = InOssMember;
 		OssPartyMember->OnMemberConnectionStatusChanged().AddUObject(this, &ThisClass::HandleMemberConnectionStatusChanged);
-		USocialToolkit* OwnerToolkit = GetParty().GetSocialManager().GetFirstLocalUserToolkit();
+		USocialToolkit* OwnerToolkit = GetParty().GetSocialManager().GetSocialToolkit(OssPartyMember->GetUserId());
 		check(OwnerToolkit);
 
 		OwnerToolkit->QueueUserDependentAction(InOssMember->GetUserId(),
diff --git a/Engine/Plugins/Online/OnlineFramework/Source/Party/Private/SocialManager.cpp b/Engine/Plugins/Online/OnlineFramework/Source/Party/Private/SocialManager.cpp
--- a/Engine/Plugins/Online/OnlineFramework/Source/Party/Private/SocialManager.cpp
+++ b/Engine/Plugins/Online/OnlineFramework/Source/Party/Private/SocialManager.cpp
@@ -272,6 +272,18 @@ USocialToolkit* USocialManager::GetSocialToolkit(int32 LocalPlayerNum) const
 	return nullptr;
 }
 
+USocialToolkit* USocialManager::GetSocialToolkit(FUniqueNetIdRepl LocalUserId) const
+{
+	for (USocialToolkit* Toolkit : SocialToolkits)
+	{
+		if (Toolkit->GetOwningLocalPlayer().GetPreferredUniqueNetId() == LocalUserId)
+		{
+			return Toolkit;
+		}
+	}
+	return nullptr;
+}
+
 void USocialManager::HandlePlatformSessionInviteAccepted(const TSharedRef<const FUniqueNetId>& LocalUserId, const FOnlineSessionSearchResult& InviteResult)
 {
 	UE_LOG(LogParty, Log, TEXT("LocalUser w/ ID [%s] has accepted platform party session invite. Attempting to join persistent party."), *LocalUserId->ToDebugString());
diff --git a/Engine/Plugins/Online/OnlineFramework/Source/Party/Public/SocialManager.h b/Engine/Plugins/Online/OnlineFramework/Source/Party/Public/SocialManager.h
--- a/Engine/Plugins/Online/OnlineFramework/Source/Party/Public/SocialManager.h
+++ b/Engine/Plugins/Online/OnlineFramework/Source/Party/Public/SocialManager.h
@@ -111,6 +111,7 @@ PARTY_SCOPE:
 
 	void NotifyPartyInitialized(USocialParty& Party);
 	USocialToolkit* GetSocialToolkit(int32 LocalPlayerNum) const;
+	USocialToolkit* GetSocialToolkit(FUniqueNetIdRepl LocalUserId) const;
 
 protected:
 	struct PARTY_API FRejoinableParty : public TSharedFromThis<FRejoinableParty>
