Fix crash on party members not being local players...

#rnx

diff --git a/Engine/Plugins/Online/OnlineFramework/Source/Party/Private/Party/PartyMember.cpp b/Engine/Plugins/Online/OnlineFramework/Source/Party/Private/Party/PartyMember.cpp
--- a/Engine/Plugins/Online/OnlineFramework/Source/Party/Private/Party/PartyMember.cpp
+++ b/Engine/Plugins/Online/OnlineFramework/Source/Party/Private/Party/PartyMember.cpp
@@ -71,6 +71,11 @@ void UPartyMember::InitializePartyMember(const FOnlinePartyMemberConstRef& InOss
 		OssPartyMember = InOssMember;
 		OssPartyMember->OnMemberConnectionStatusChanged().AddUObject(this, &ThisClass::HandleMemberConnectionStatusChanged);
 		USocialToolkit* OwnerToolkit = GetParty().GetSocialManager().GetSocialToolkit(OssPartyMember->GetUserId());
+		// If we are not a local user then we simply get the first local user's toolkit
+		if (OwnerToolkit == nullptr)
+		{
+			OwnerToolkit = GetParty().GetSocialManager().GetFirstLocalUserToolkit();
+		}
 		check(OwnerToolkit);
 
 		OwnerToolkit->QueueUserDependentAction(InOssMember->GetUserId(),
