Fix issue where local player couldn't update their own party member data

diff --git a/Engine/Plugins/Online/OnlineFramework/Source/Party/Private/Party/PartyDataReplicator.cpp b/Engine/Plugins/Online/OnlineFramework/Source/Party/Private/Party/PartyDataReplicator.cpp
--- a/Engine/Plugins/Online/OnlineFramework/Source/Party/Private/Party/PartyDataReplicator.cpp
+++ b/Engine/Plugins/Online/OnlineFramework/Source/Party/Private/Party/PartyDataReplicator.cpp
@@ -21,8 +21,17 @@ void FPartyDataReplicatorHelper::ReplicateDataToMembers(const FOnlinePartyRepDat
 			}
 			else if (RepDataType.IsChildOf(FPartyMemberRepData::StaticStruct()))
 			{
-				UE_LOG(LogParty, VeryVerbose, TEXT("Sending rep data update for member within party [%s]."), *OwnerParty->ToDebugString());
-				PartyInterface->UpdatePartyMemberData(*LocalUserId, OwnerParty->GetPartyId(), ReplicationPayload);
+				if (const UPartyMember* Owner = static_cast<const FOnlinePartyRepDataBase*>(&RepDataInstance)->GetOwningMember())
+				{
+					LocalUserId = Owner->GetPrimaryNetId();
+					UE_LOG(LogParty, VeryVerbose, TEXT("Sending rep data update for member within party [%s]."), *OwnerParty->ToDebugString());
+					PartyInterface->UpdatePartyMemberData(*LocalUserId, OwnerParty->GetPartyId(), ReplicationPayload);
+				}
+				else
+				{
+					UE_LOG(LogParty, Warning, TEXT("Sending rep data update for member within party [%s] Could not identidy owner."), *OwnerParty->ToDebugString());
+					PartyInterface->UpdatePartyMemberData(*LocalUserId, OwnerParty->GetPartyId(), ReplicationPayload);
+				}				
 			}
 		}
 	}
diff --git a/Engine/Plugins/Online/OnlineFramework/Source/Party/Private/SocialManager.cpp b/Engine/Plugins/Online/OnlineFramework/Source/Party/Private/SocialManager.cpp
--- a/Engine/Plugins/Online/OnlineFramework/Source/Party/Private/SocialManager.cpp
+++ b/Engine/Plugins/Online/OnlineFramework/Source/Party/Private/SocialManager.cpp
@@ -999,7 +999,7 @@ void USocialManager::HandleCreatePartyComplete(const FUniqueNetId& LocalUserId,
 
 void USocialManager::HandleJoinPartyComplete(const FUniqueNetId& LocalUserId, const FOnlinePartyId& PartyId, EJoinPartyCompletionResult Result, int32 NotApprovedReasonCode, FOnlinePartyTypeId PartyTypeId)
 {
-	UE_LOG(LogParty, Log, TEXT("Attempt to join party of type [%d] completed with result [%s] and reason code [%d]"), PartyTypeId.GetValue(), ToString(Result), NotApprovedReasonCode);
+	UE_LOG(LogParty, Log, TEXT("Attempt to join party of type [%d] completed with result [%s] and reason code [%d] by user is [%s] "), PartyTypeId.GetValue(), ToString(Result), NotApprovedReasonCode, *LocalUserId.ToDebugString());
 
 	FJoinPartyResult JoinResult(Result, NotApprovedReasonCode);
 	FJoinPartyAttempt* JoinAttempt = JoinAttemptsByTypeId.Find(PartyTypeId);
diff --git a/Engine/Plugins/Online/OnlineFramework/Source/Party/Public/Party/PartyTypes.h b/Engine/Plugins/Online/OnlineFramework/Source/Party/Public/Party/PartyTypes.h
--- a/Engine/Plugins/Online/OnlineFramework/Source/Party/Public/Party/PartyTypes.h
+++ b/Engine/Plugins/Online/OnlineFramework/Source/Party/Public/Party/PartyTypes.h
@@ -412,6 +412,7 @@ protected:
 	virtual void PostReplication() {}
 	virtual bool CanEditData() const { return false; }
 	virtual const USocialParty* GetOwnerParty() const { return nullptr; }
+	virtual const UPartyMember* GetOwningMember() const { return nullptr; }
 
 	void LogSetPropertyFailure(const TCHAR* OwningStructTypeName, const TCHAR* PropertyName) const;
 	void LogPropertyChanged(const TCHAR* OwningStructTypeName, const TCHAR* PropertyName, bool bFromReplication) const;
