Moved NotifyPartyInitialized out of PARTY_SCOPE so game subclasses of SocialManager can override it.

diff --git a/Engine/Plugins/Online/OnlineFramework/Source/Party/Public/SocialManager.h b/Engine/Plugins/Online/OnlineFramework/Source/Party/Public/SocialManager.h
--- a/Engine/Plugins/Online/OnlineFramework/Source/Party/Public/SocialManager.h
+++ b/Engine/Plugins/Online/OnlineFramework/Source/Party/Public/SocialManager.h
@@ -102,6 +102,8 @@ public:
 	 */
 	void RegisterSecondaryPlayer(int32 LocalPlayerNum, const FOnJoinPartyComplete& Delegate = FOnJoinPartyComplete());
 
+	virtual void NotifyPartyInitialized(USocialParty& Party);
+
 PARTY_SCOPE:
 	/** Validates that the target user has valid join info for us to use and that we can join any party of the given type */
 	FJoinPartyResult ValidateJoinTarget(const USocialUser& UserToJoin, const FOnlinePartyTypeId& PartyTypeId) const;
@@ -109,7 +111,6 @@ PARTY_SCOPE:
 	DECLARE_DELEGATE_OneParam(FOnJoinPartyAttemptComplete, const FJoinPartyResult&);
 	void JoinParty(const USocialUser& UserToJoin, const FOnlinePartyTypeId& PartyTypeId, const FOnJoinPartyAttemptComplete& OnJoinPartyComplete);
 
-	void NotifyPartyInitialized(USocialParty& Party);
 	USocialToolkit* GetSocialToolkit(int32 LocalPlayerNum) const;
 	USocialToolkit* GetSocialToolkit(FUniqueNetIdRepl LocalUserId) const;
 
