Add hotfixable delay when querying userinfo, in order to aggregate more queries into a single HTTP call.

diff --git a/Engine/Plugins/Online/OnlineFramework/Source/Party/Public/SocialQuery.h b/Engine/Plugins/Online/OnlineFramework/Source/Party/Public/SocialQuery.h
--- a/Engine/Plugins/Online/OnlineFramework/Source/Party/Public/SocialQuery.h
+++ b/Engine/Plugins/Online/OnlineFramework/Source/Party/Public/SocialQuery.h
@@ -104,6 +104,8 @@ public:
 		}
 
 		TickExecuteHandle.Reset();
+
+		// Returning false ensures the ticker removes this delegate
 		return false;
 	}
 
@@ -135,10 +137,14 @@ private:
 		NewQuery->Initialize(Toolkit, SubsystemType, FOnQueryCompleted::CreateRaw(this, &FSocialQueryManager::HandleQueryComplete));
 		Queries.Add(NewQuery);
 
+		float UserInfoQueryAggregationTime = 0.0f;
+		
+		GConfig->GetFloat(TEXT("Social"), TEXT("UserInfoQueryAggregationTime"), UserInfoQueryAggregationTime, GGameIni);
+
 		// If we aren't already registered to execute our queries next tick, do so now
 		if (!TickExecuteHandle.IsValid())
 		{
-			TickExecuteHandle = FTicker::GetCoreTicker().AddTicker(FTickerDelegate::CreateRaw(this, &FSocialQueryManager::HandleExecuteQueries));
+			TickExecuteHandle = FTicker::GetCoreTicker().AddTicker(FTickerDelegate::CreateRaw(this, &FSocialQueryManager::HandleExecuteQueries), UserInfoQueryAggregationTime);
 		}
 
 		return NewQuery;
