Party join approval support for a group of players joining together
Change party join request received delegates to take in a list of users.  ApproveJoinRequest only needs to be called for one of the joining users.

diff --git a/Engine/Plugins/Online/OnlineFramework/Source/Party/Private/Party/SocialParty.cpp b/Engine/Plugins/Online/OnlineFramework/Source/Party/Private/Party/SocialParty.cpp
--- a/Engine/Plugins/Online/OnlineFramework/Source/Party/Private/Party/SocialParty.cpp
+++ b/Engine/Plugins/Online/OnlineFramework/Source/Party/Private/Party/SocialParty.cpp
@@ -186,11 +186,11 @@ bool USocialParty::ApplyCrossplayRestriction(FPartyJoinApproval& JoinApproval, c
 	return !JoinApproval.GetDenialReason().HasAnyReason();
 }
 
-FPartyJoinApproval USocialParty::EvaluateJoinRequest(const FUniqueNetId& PlayerId, const FUserPlatform& Platform, const FOnlinePartyData& JoinData, bool bFromJoinRequest) const
+FPartyJoinApproval USocialParty::EvaluateJoinRequest(const TArray<IOnlinePartyUserPendingJoinRequestInfoConstRef>& Players, bool bFromJoinRequest) const
 {
 	FPartyJoinApproval JoinApproval;
 
-	if (IsPartyFull())
+	if ((GetNumPartyMembers() + Players.Num()) > GetPartyMaxSize())
 	{
 		JoinApproval.SetDenialReason(EPartyJoinDenialReason::PartyFull);
 	}
@@ -205,7 +205,14 @@ FPartyJoinApproval USocialParty::EvaluateJoinRequest(const FUniqueNetId& PlayerI
 		{
 			if (bAlwaysCheckCrossplatformOnPartyJoin)
 			{
-				ApplyCrossplayRestriction(JoinApproval, Platform, JoinData);
+				for (const IOnlinePartyUserPendingJoinRequestInfoConstRef& Player : Players)
+				{
+					FUserPlatform Platform(Player->GetPlatform());
+					if (!ApplyCrossplayRestriction(JoinApproval, Platform, *Player->GetJoinData()))
+					{
+						break;
+					}
+				}
 			}
 		}
 
@@ -413,9 +420,9 @@ void USocialParty::InitializePartyInternal()
 	IOnlinePartyPtr PartyInterface = Online::GetPartyInterfaceChecked(GetWorld());
 	PartyInterface->AddOnPartyConfigChangedDelegate_Handle(FOnPartyConfigChangedConstDelegate::CreateUObject(this, &USocialParty::HandlePartyConfigChanged));
 	PartyInterface->AddOnPartyDataReceivedDelegate_Handle(FOnPartyDataReceivedConstDelegate::CreateUObject(this, &USocialParty::HandlePartyDataReceived));
-	PartyInterface->AddOnPartyJoinRequestReceivedDelegate_Handle(FOnPartyJoinRequestReceivedDelegate::CreateUObject(this, &USocialParty::HandlePartyJoinRequestReceived));
+	PartyInterface->AddOnPartyJoinRequestReceivedDelegate_Handle(FOnPartyGroupJoinRequestReceivedDelegate::CreateUObject(this, &USocialParty::HandlePartyJoinRequestReceived));
 	PartyInterface->AddOnPartyJIPRequestReceivedDelegate_Handle(FOnPartyJIPRequestReceivedDelegate::CreateUObject(this, &USocialParty::HandlePartyJIPRequestReceived));
-	PartyInterface->AddOnQueryPartyJoinabilityReceivedDelegate_Handle(FOnQueryPartyJoinabilityReceivedDelegate::CreateUObject(this, &USocialParty::HandleJoinabilityQueryReceived));
+	PartyInterface->AddOnQueryPartyJoinabilityReceivedDelegate_Handle(FOnQueryPartyJoinabilityGroupReceivedDelegate::CreateUObject(this, &USocialParty::HandleJoinabilityQueryReceived));
 	PartyInterface->AddOnPartyExitedDelegate_Handle(FOnPartyExitedDelegate::CreateUObject(this, &USocialParty::HandlePartyLeft));
 	PartyInterface->AddOnPartyStateChangedDelegate_Handle(FOnPartyStateChangedDelegate::CreateUObject(this, &USocialParty::HandlePartyStateChanged));
 
@@ -579,36 +586,43 @@ UPartyMember* USocialParty::GetOrCreatePartyMember(const FUniqueNetId& MemberId)
 	return PartyMember;
 }
 
-void USocialParty::HandlePartyJoinRequestReceived(const FUniqueNetId& LocalUserId, const FOnlinePartyId& PartyId, const FUniqueNetId& SenderId, const FString& Platform, const FOnlinePartyData& JoinData)
+void USocialParty::HandlePartyJoinRequestReceived(const FUniqueNetId& LocalUserId, const FOnlinePartyId& PartyId, const IOnlinePartyPendingJoinRequestInfo& JoinRequestInfo)
 {
 	if (!IsLocalPlayerPartyLeader() || PartyId != GetPartyId())
 	{
 		return;
 	}
 
+	TArray<IOnlinePartyUserPendingJoinRequestInfoConstRef> JoiningUsers;
+	JoinRequestInfo.GetUsers(JoiningUsers);
+	check(JoiningUsers.IsValidIndex(0));
+
+	const IOnlinePartyUserPendingJoinRequestInfoConstRef& PrimaryJoiningUser = JoiningUsers[0];
+
 	IOnlinePartyPtr PartyInterface = Online::GetPartyInterfaceChecked(GetWorld());
 #if !UE_BUILD_SHIPPING
 	if (AutoApproveJoinRequests != 0)
 	{
-		PartyInterface->ApproveJoinRequest(LocalUserId, PartyId, SenderId, true);
+		PartyInterface->ApproveJoinRequest(LocalUserId, PartyId, *PrimaryJoiningUser->GetUserId(), true);
 		return;
 	}
 #endif
 
-	const FUserPlatform MemberPlatform(Platform);
-	FPartyJoinApproval JoinApproval = EvaluateJoinRequest(SenderId, MemberPlatform, JoinData, true);
+	FPartyJoinApproval JoinApproval = EvaluateJoinRequest(JoiningUsers, true);
 
 	if (JoinApproval.GetApprovalAction() == EApprovalAction::Enqueue ||
 		JoinApproval.GetApprovalAction() == EApprovalAction::EnqueueAndStartBeacon)
 	{
 		// Enqueue for a more opportune time
-		UE_LOG(LogParty, Verbose, TEXT("[%s] Enqueuing approval request for %s"), *PartyId.ToString(), *SenderId.ToString());
+		UE_LOG(LogParty, Verbose, TEXT("[%s] Enqueuing approval request for %s"), *PartyId.ToString(), *PrimaryJoiningUser->GetUserId()->ToDebugString());
 		
 		FPendingMemberApproval PendingApproval;
 		PendingApproval.RecipientId.SetUniqueNetId(LocalUserId.AsShared());
-		PendingApproval.SenderId.SetUniqueNetId(SenderId.AsShared());
-		PendingApproval.Platform = MemberPlatform;
-		PendingApproval.JoinData = JoinData.AsShared();
+		PendingApproval.Members.Reserve(JoiningUsers.Num());
+		for (IOnlinePartyUserPendingJoinRequestInfoConstRef JoiningUser : JoiningUsers)
+		{
+			PendingApproval.Members.Emplace(JoiningUser->GetUserId(), JoiningUser->GetPlatform(), JoiningUser->GetJoinData());
+		}
 		PendingApproval.bIsJIPApproval = false;
 		PendingApprovals.Enqueue(PendingApproval);
 
@@ -620,9 +634,9 @@ void USocialParty::HandlePartyJoinRequestReceived(const FUniqueNetId& LocalUserI
 	else
 	{
 		const bool bIsApproved = JoinApproval.CanJoin();
-		UE_LOG(LogParty, Verbose, TEXT("[%s] Responding to approval request for %s with %s"), *PartyId.ToString(), *SenderId.ToString(), bIsApproved ? TEXT("approved") : TEXT("denied"));
+		UE_LOG(LogParty, Verbose, TEXT("[%s] Responding to approval request for %s with %s"), *PartyId.ToString(), *PrimaryJoiningUser->GetUserId()->ToDebugString(), bIsApproved ? TEXT("approved") : TEXT("denied"));
 		
-		PartyInterface->ApproveJoinRequest(LocalUserId, PartyId, SenderId, bIsApproved, JoinApproval.GetDenialReason());
+		PartyInterface->ApproveJoinRequest(LocalUserId, PartyId, *PrimaryJoiningUser->GetUserId(), bIsApproved, JoinApproval.GetDenialReason());
 	}
 }
 
@@ -646,8 +660,7 @@ void USocialParty::RemovePlayerFromReservationBeacon(const FUniqueNetId& LocalUs
 
 	FPendingMemberApproval PendingApproval;
 	PendingApproval.RecipientId.SetUniqueNetId(LocalUserId.AsShared());
-	PendingApproval.SenderId.SetUniqueNetId(PlayerToRemove.AsShared());
-	PendingApproval.Platform = MemberPlatform;
+	PendingApproval.Members.Emplace(PlayerToRemove.AsShared(), MemberPlatform);
 	PendingApproval.bIsJIPApproval = true;
 	PendingApproval.bIsPlayerRemoval = true;
 	PendingApprovals.Enqueue(PendingApproval);
@@ -689,8 +702,7 @@ void USocialParty::HandlePartyJIPRequestReceived(const FUniqueNetId& LocalUserId
 
 		FPendingMemberApproval PendingApproval;
 		PendingApproval.RecipientId.SetUniqueNetId(LocalUserId.AsShared());
-		PendingApproval.SenderId.SetUniqueNetId(SenderId.AsShared());
-		PendingApproval.Platform = MemberPlatform;
+		PendingApproval.Members.Emplace(SenderId.AsShared(), MemberPlatform);
 		PendingApproval.bIsJIPApproval = true;
 		PendingApprovals.Enqueue(PendingApproval);
 
@@ -709,15 +721,21 @@ void USocialParty::HandlePartyJIPRequestReceived(const FUniqueNetId& LocalUserId
 }
 
 
-void USocialParty::HandleJoinabilityQueryReceived(const FUniqueNetId& LocalUserId, const FOnlinePartyId& PartyId, const FUniqueNetId& SenderId, const FString& Platform, const FOnlinePartyData& JoinData)
+void USocialParty::HandleJoinabilityQueryReceived(const FUniqueNetId& LocalUserId, const FOnlinePartyId& PartyId, const IOnlinePartyPendingJoinRequestInfo& JoinRequestInfo)
 {
-	if (PartyId == GetPartyId())
+	if (PartyId == GetPartyId()) 
 	{
-		FPartyJoinApproval JoinabilityInfo = EvaluateJoinRequest(SenderId, FUserPlatform(Platform), JoinData, false);
-		UE_LOG(LogParty, VeryVerbose, TEXT("[%s] Responding to approval request for %s with %s"), *PartyId.ToString(), *SenderId.ToString(), JoinabilityInfo.CanJoin() ? TEXT("approved") : TEXT("denied"));
+		TArray<IOnlinePartyUserPendingJoinRequestInfoConstRef> JoiningUsers;
+		JoinRequestInfo.GetUsers(JoiningUsers);
+		check(JoiningUsers.IsValidIndex(0));
+
+		const IOnlinePartyUserPendingJoinRequestInfoConstRef& PrimaryJoiningUser = JoiningUsers[0];
+
+		FPartyJoinApproval JoinabilityInfo = EvaluateJoinRequest(JoiningUsers, false);
+		UE_LOG(LogParty, VeryVerbose, TEXT("[%s] Responding to approval request for %s with %s"), *PartyId.ToString(), *PrimaryJoiningUser->GetUserId()->ToString(), JoinabilityInfo.CanJoin() ? TEXT("approved") : TEXT("denied"));
 
 		const IOnlinePartyPtr PartyInterface = Online::GetPartyInterfaceChecked(GetWorld());
-		PartyInterface->RespondToQueryJoinability(LocalUserId, PartyId, SenderId, JoinabilityInfo.CanJoin(), JoinabilityInfo.GetDenialReason());
+		PartyInterface->RespondToQueryJoinability(LocalUserId, PartyId, *PrimaryJoiningUser->GetUserId(), JoinabilityInfo.CanJoin(), JoinabilityInfo.GetDenialReason());		
 	}
 }
 
@@ -1279,22 +1297,25 @@ void USocialParty::ConnectToReservationBeacon()
 							ReservationBeaconClient->OnHostConnectionFailure().BindUObject(this, &USocialParty::HandleBeaconHostConnectionFailed);
 							ReservationBeaconClient->OnReservationRequestComplete().BindUObject(this, &USocialParty::HandleReservationRequestComplete);
 
-							FPlayerReservation Reservation;
-							Reservation.UniqueId = NextApproval.SenderId;
-							Reservation.Platform = NextApproval.Platform;
-
-							if (!NextApproval.bIsJIPApproval)
-							{
-								const ECrossplayPreference CrossplayPreference = GetCrossplayPreferenceFromJoinData(*NextApproval.JoinData);
-								Reservation.bAllowCrossplay = (CrossplayPreference == ECrossplayPreference::OptedIn);
-							}
-							else
+							TArray<FPlayerReservation> ReservationAsArray;
+							ReservationAsArray.Reserve(NextApproval.Members.Num());
+							for (const FPendingMemberApproval::FMemberInfo& MemberInfo : NextApproval.Members)
 							{
-								Reservation.bAllowCrossplay = true; // This will not matter since we are JIP, and the session already has crossplay set.
+								FPlayerReservation Reservation;
+								Reservation.UniqueId = MemberInfo.MemberId;
+								Reservation.Platform = MemberInfo.Platform;
+
+								if (!NextApproval.bIsJIPApproval && MemberInfo.JoinData.IsValid())
+								{
+									const ECrossplayPreference CrossplayPreference = GetCrossplayPreferenceFromJoinData(*MemberInfo.JoinData);
+									Reservation.bAllowCrossplay = (CrossplayPreference == ECrossplayPreference::OptedIn);
+								}
+								else
+								{
+									Reservation.bAllowCrossplay = true; // This will not matter since we are JIP, and the session already has crossplay set.
+								}
 							}
 
-							TArray<FPlayerReservation> ReservationAsArray;
-							ReservationAsArray.Add(Reservation);
 							bStartedConnection = ReservationBeaconClient->RequestReservationUpdate(URL, Session->GetSessionIdStr(), GetPartyLeader()->GetPrimaryNetId(), ReservationAsArray, NextApproval.bIsPlayerRemoval);
 						}
 					}
@@ -1318,14 +1339,15 @@ void USocialParty::RejectAllPendingJoinRequests()
 	while (!PendingApprovals.IsEmpty())
 	{
 		PendingApprovals.Dequeue(PendingApproval);
-		UE_LOG(LogParty, Verbose, TEXT("[%s] Responding to approval request for %s with denied"), *PartyId.ToString(), *PendingApproval.SenderId.ToString());
+		const FUniqueNetId& PrimaryJoiningUserId = *PendingApproval.Members[0].MemberId.GetUniqueNetId();
+		UE_LOG(LogParty, Verbose, TEXT("[%s] Responding to approval request for %s with denied"), *PartyId.ToString(), *PrimaryJoiningUserId.ToDebugString());
 		if (PendingApproval.bIsJIPApproval)
 		{
-			PartyInterface->ApproveJIPRequest(*PendingApproval.RecipientId, PartyId, *PendingApproval.SenderId, false, (int32)EPartyJoinDenialReason::Busy);
+			PartyInterface->ApproveJIPRequest(*PendingApproval.RecipientId, PartyId, PrimaryJoiningUserId, false, (int32)EPartyJoinDenialReason::Busy);
 		}
 		else
 		{
-			PartyInterface->ApproveJoinRequest(*PendingApproval.RecipientId, PartyId, *PendingApproval.SenderId, false, (int32)EPartyJoinDenialReason::Busy);
+			PartyInterface->ApproveJoinRequest(*PendingApproval.RecipientId, PartyId, PrimaryJoiningUserId, false, (int32)EPartyJoinDenialReason::Busy);
 		}
 	}
 }
@@ -1371,36 +1393,29 @@ void USocialParty::PumpApprovalQueue()
 	{
 		if (ensure(ReservationBeaconClient))
 		{
-			if (NextApproval.bIsJIPApproval == false)
-			{
-				// This is a request to join our party
-				ECrossplayPreference CrossplayPreference = GetCrossplayPreferenceFromJoinData(*NextApproval.JoinData);
-
-				FPlayerReservation NewPlayerRes;
-				NewPlayerRes.UniqueId = NextApproval.SenderId;
-				NewPlayerRes.Platform = NextApproval.Platform;
-				NewPlayerRes.bAllowCrossplay = (CrossplayPreference == ECrossplayPreference::OptedIn);
-
-				TArray<FPlayerReservation> PlayersToAdd;
-				PlayersToAdd.Add(NewPlayerRes);
-
-				ReservationBeaconClient->RequestReservationUpdate(GetPartyLeader()->GetPrimaryNetId(), PlayersToAdd);
-			}
-			else
+			TArray<FPlayerReservation> PlayersToAdd;
+			PlayersToAdd.Reserve(NextApproval.Members.Num());
+			for (const FPendingMemberApproval::FMemberInfo& MemberInfo : NextApproval.Members)
 			{
-				// This is a request from a party member to join a JIP game.
 				FPlayerReservation NewPlayerRes;
-				NewPlayerRes.UniqueId = NextApproval.SenderId;
-				NewPlayerRes.Platform = NextApproval.Platform;
+				NewPlayerRes.UniqueId = MemberInfo.MemberId;
+				NewPlayerRes.Platform = MemberInfo.Platform;
 
-				// This doesn't matter, since the crossplay state of the match has already been set.
-				NewPlayerRes.bAllowCrossplay = true;
-
-				TArray<FPlayerReservation> PlayersToAdd;
-				PlayersToAdd.Add(NewPlayerRes);
-
-				ReservationBeaconClient->RequestReservationUpdate(GetPartyLeader()->GetPrimaryNetId(), PlayersToAdd);
+				if (NextApproval.bIsJIPApproval == false && MemberInfo.JoinData.IsValid())
+				{
+					// This is a request to join our party
+					ECrossplayPreference CrossplayPreference = GetCrossplayPreferenceFromJoinData(*MemberInfo.JoinData);
+					NewPlayerRes.bAllowCrossplay = (CrossplayPreference == ECrossplayPreference::OptedIn);
+				}
+				else
+				{
+					// This is a request from a party member to join a JIP game.
+					// This doesn't matter, since the crossplay state of the match has already been set.
+					NewPlayerRes.bAllowCrossplay = true;
+				}
+				PlayersToAdd.Emplace(MoveTemp(NewPlayerRes));
 			}
+			ReservationBeaconClient->RequestReservationUpdate(GetPartyLeader()->GetPrimaryNetId(), PlayersToAdd);
 		}
 		else
 		{
@@ -1431,7 +1446,7 @@ void USocialParty::HandleReservationRequestComplete(EPartyReservationResult::Typ
 			if (PendingApproval.bIsJIPApproval)
 			{
 				// This player is already in our party. ApproveJIPRequest
-				PartyInterface->ApproveJIPRequest(*PendingApproval.RecipientId, GetPartyId(), *PendingApproval.SenderId, bReservationApproved, DenialReason);
+				PartyInterface->ApproveJIPRequest(*PendingApproval.RecipientId, GetPartyId(), *PendingApproval.Members[0].MemberId, bReservationApproved, DenialReason);
 			}
 			else if (PendingApproval.bIsPlayerRemoval)
 			{
@@ -1439,7 +1454,7 @@ void USocialParty::HandleReservationRequestComplete(EPartyReservationResult::Typ
 			}
 			else
 			{
-				PartyInterface->ApproveJoinRequest(*PendingApproval.RecipientId, GetPartyId(), *PendingApproval.SenderId, bReservationApproved, DenialReason);
+				PartyInterface->ApproveJoinRequest(*PendingApproval.RecipientId, GetPartyId(), *PendingApproval.Members[0].MemberId, bReservationApproved, DenialReason);
 			}
 		}
 		PumpApprovalQueue();
diff --git a/Engine/Plugins/Online/OnlineFramework/Source/Party/Private/SocialDebugTools.cpp b/Engine/Plugins/Online/OnlineFramework/Source/Party/Private/SocialDebugTools.cpp
--- a/Engine/Plugins/Online/OnlineFramework/Source/Party/Private/SocialDebugTools.cpp
+++ b/Engine/Plugins/Online/OnlineFramework/Source/Party/Private/SocialDebugTools.cpp
@@ -662,7 +662,7 @@ void USocialDebugTools::FInstanceContext::Init()
 		if (OnlineParty.IsValid())
 		{
 			PartyInviteReceivedDelegateHandle = OnlineParty->AddOnPartyInviteReceivedDelegate_Handle(FOnPartyInviteReceivedDelegate::CreateUObject(&Owner, &USocialDebugTools::HandlePartyInviteReceived));
-			PartyJoinRequestReceivedDelegateHandle = OnlineParty->AddOnPartyJoinRequestReceivedDelegate_Handle(FOnPartyJoinRequestReceivedDelegate::CreateUObject(&Owner, &USocialDebugTools::HandlePartyJoinRequestReceived));
+			PartyJoinRequestReceivedDelegateHandle = OnlineParty->AddOnPartyJoinRequestReceivedDelegate_Handle(FOnPartyGroupJoinRequestReceivedDelegate::CreateUObject(&Owner, &USocialDebugTools::HandlePartyJoinRequestReceived));
 		}
 	}
 }
@@ -749,7 +749,7 @@ void USocialDebugTools::HandlePartyInviteReceived(const FUniqueNetId& LocalUserI
 	}
 }
 
-void USocialDebugTools::HandlePartyJoinRequestReceived(const FUniqueNetId& LocalUserId, const FOnlinePartyId& PartyId, const FUniqueNetId& SenderId, const FString& Platform, const FOnlinePartyData& PartyData)
+void USocialDebugTools::HandlePartyJoinRequestReceived(const FUniqueNetId& LocalUserId, const FOnlinePartyId& PartyId, const IOnlinePartyPendingJoinRequestInfo& JoinRequestInfo)
 {
 	FInstanceContext* Context = GetContextForUser(LocalUserId);
 	if (Context)
@@ -757,7 +757,10 @@ void USocialDebugTools::HandlePartyJoinRequestReceived(const FUniqueNetId& Local
 		IOnlinePartyPtr OnlineParty = Context->GetOSS()->GetPartyInterface();
 		if (OnlineParty.IsValid())
 		{
-			OnlineParty->ApproveJoinRequest(LocalUserId, PartyId, SenderId, true);
+			TArray<IOnlinePartyUserPendingJoinRequestInfoConstRef> JoiningUsers;
+			JoinRequestInfo.GetUsers(JoiningUsers);
+			check(JoiningUsers.IsValidIndex(0));
+			OnlineParty->ApproveJoinRequest(LocalUserId, PartyId, *JoiningUsers[0]->GetUserId(), true);
 			return;
 		}
 	}
diff --git a/Engine/Plugins/Online/OnlineFramework/Source/Party/Public/Party/SocialParty.h b/Engine/Plugins/Online/OnlineFramework/Source/Party/Public/Party/SocialParty.h
--- a/Engine/Plugins/Online/OnlineFramework/Source/Party/Public/Party/SocialParty.h
+++ b/Engine/Plugins/Online/OnlineFramework/Source/Party/Public/Party/SocialParty.h
@@ -236,8 +236,8 @@ protected:
 	
 	virtual void HandlePartySystemStateChange(EPartySystemState NewState);
 
-	/** Determines the joinability of this party for a specific user requesting to join */
-	virtual FPartyJoinApproval EvaluateJoinRequest(const FUniqueNetId& PlayerId, const FUserPlatform& Platform, const FOnlinePartyData& JoinData, bool bFromJoinRequest) const;
+	/** Determines the joinability of this party for a group of users requesting to join */
+	virtual FPartyJoinApproval EvaluateJoinRequest(const TArray<IOnlinePartyUserPendingJoinRequestInfoConstRef>& Players, bool bFromJoinRequest) const;
 
 	/** Determines the joinability of the game a party is in for JoinInProgress */
 	virtual FPartyJoinApproval EvaluateJIPRequest(const FUniqueNetId& PlayerId) const;
@@ -307,8 +307,8 @@ private:	// Handlers
 	void HandlePartyConfigChanged(const FUniqueNetId& LocalUserId, const FOnlinePartyId& PartyId, const FPartyConfiguration& PartyConfig);
 	void HandleUpdatePartyConfigComplete(const FUniqueNetId& LocalUserId, const FOnlinePartyId& PartyId, EUpdateConfigCompletionResult Result);
 	void HandlePartyDataReceived(const FUniqueNetId& LocalUserId, const FOnlinePartyId& PartyId, const FOnlinePartyData& PartyData);
-	void HandleJoinabilityQueryReceived(const FUniqueNetId& LocalUserId, const FOnlinePartyId& PartyId, const FUniqueNetId& SenderId, const FString& Platform, const FOnlinePartyData& JoinData);
-	void HandlePartyJoinRequestReceived(const FUniqueNetId& LocalUserId, const FOnlinePartyId& PartyId, const FUniqueNetId& SenderId, const FString& Platform, const FOnlinePartyData& JoinData);
+	void HandleJoinabilityQueryReceived(const FUniqueNetId& LocalUserId, const FOnlinePartyId& PartyId, const IOnlinePartyPendingJoinRequestInfo& JoinRequestInfo);
+	void HandlePartyJoinRequestReceived(const FUniqueNetId& LocalUserId, const FOnlinePartyId& PartyId, const IOnlinePartyPendingJoinRequestInfo& JoinRequestInfo);
 	void HandlePartyJIPRequestReceived(const FUniqueNetId& LocalUserId, const FOnlinePartyId& PartyId, const FUniqueNetId& SenderId);
 	void HandlePartyLeft(const FUniqueNetId& LocalUserId, const FOnlinePartyId& PartyId);
 	void HandlePartyMemberExited(const FUniqueNetId& LocalUserId, const FOnlinePartyId& PartyId, const FUniqueNetId& MemberId, EMemberExitedReason ExitReason);
@@ -348,12 +348,23 @@ private:
 	//@todo DanH Party: Rename/reorg this to more clearly call out that this is specific to lobby beacon stuff #suggested
 	struct FPendingMemberApproval
 	{
+		struct FMemberInfo
+		{
+			FMemberInfo(FUniqueNetIdRepl InMemberId, FUserPlatform InPlatform, TSharedPtr<const FOnlinePartyData> InJoinData = TSharedPtr<const FOnlinePartyData>())
+				: MemberId(InMemberId)
+				, Platform(MoveTemp(InPlatform))
+				, JoinData(InJoinData)
+			{}
+
+			FUniqueNetIdRepl MemberId;
+			FUserPlatform Platform;
+			TSharedPtr<const FOnlinePartyData> JoinData;
+		};
+
 		FUniqueNetIdRepl RecipientId;
-		FUniqueNetIdRepl SenderId;
-		FUserPlatform Platform;
+		TArray<FMemberInfo> Members;
 		bool bIsJIPApproval;
 		bool bIsPlayerRemoval = false;
-		TSharedPtr<const FOnlinePartyData> JoinData;
 	};
 	TQueue<FPendingMemberApproval> PendingApprovals;
 
diff --git a/Engine/Plugins/Online/OnlineFramework/Source/Party/Public/SocialDebugTools.h b/Engine/Plugins/Online/OnlineFramework/Source/Party/Public/SocialDebugTools.h
--- a/Engine/Plugins/Online/OnlineFramework/Source/Party/Public/SocialDebugTools.h
+++ b/Engine/Plugins/Online/OnlineFramework/Source/Party/Public/SocialDebugTools.h
@@ -8,6 +8,7 @@
 
 class IOnlineSubsystem;
 class FOnlineAccountCredentials;
+class IOnlinePartyPendingJoinRequestInfo;
 typedef TSharedPtr<const class IOnlinePartyJoinInfo> IOnlinePartyJoinInfoConstPtr;
 typedef TSharedPtr<class FOnlinePartyData> FOnlinePartyDataPtr;
 typedef TSharedPtr<const class FOnlinePartyData> FOnlinePartyDataConstPtr;
@@ -100,5 +101,5 @@ private:
 	// OSS callback handlers
 	void HandleFriendInviteReceived(const FUniqueNetId& LocalUserId, const FUniqueNetId& FriendId);
 	void HandlePartyInviteReceived(const FUniqueNetId& LocalUserId, const FOnlinePartyId& PartyId, const FUniqueNetId& SenderId);
-	void HandlePartyJoinRequestReceived(const FUniqueNetId& LocalUserId, const FOnlinePartyId& PartyId, const FUniqueNetId& SenderId, const FString& Platform, const FOnlinePartyData& PartyData);
+	void HandlePartyJoinRequestReceived(const FUniqueNetId& LocalUserId, const FOnlinePartyId& PartyId, const IOnlinePartyPendingJoinRequestInfo& JoinRequestInfo);
 };
\ No newline at end of file
diff --git a/Engine/Plugins/Online/OnlineSubsystem/Source/Private/OnlinePartyInterface.cpp b/Engine/Plugins/Online/OnlineSubsystem/Source/Private/OnlinePartyInterface.cpp
--- a/Engine/Plugins/Online/OnlineSubsystem/Source/Private/OnlinePartyInterface.cpp
+++ b/Engine/Plugins/Online/OnlineSubsystem/Source/Private/OnlinePartyInterface.cpp
@@ -213,6 +213,30 @@ FDelegateHandle IOnlinePartySystem::AddOnPartyMemberDataReceivedDelegate_Handle(
 	return OnPartyMemberDataReceivedDelegates.Add(FOnPartyMemberDataReceivedConstDelegate::CreateLambda(DeprecationHelperLambda));
 }
 
+FDelegateHandle IOnlinePartySystem::AddOnPartyJoinRequestReceivedDelegate_Handle(const FOnPartyJoinRequestReceivedDelegate& Delegate)
+{
+	auto DeprecationHelperLambda = [Delegate](const FUniqueNetId& LocalUserId, const FOnlinePartyId& PartyId, const IOnlinePartyPendingJoinRequestInfo& JoiningUsers)
+	{
+		TArray<IOnlinePartyUserPendingJoinRequestInfoConstRef> Users;
+		JoiningUsers.GetUsers(Users);
+		check(Users.Num() > 0);
+		Delegate.ExecuteIfBound(LocalUserId, PartyId, *Users[0]->GetUserId(), Users[0]->GetPlatform(), *Users[0]->GetJoinData());
+	};
+	return OnPartyJoinRequestReceivedDelegates.Add(FOnPartyGroupJoinRequestReceivedDelegate::CreateLambda(DeprecationHelperLambda));
+}
+
+FDelegateHandle IOnlinePartySystem::AddOnQueryPartyJoinabilityReceivedDelegate_Handle(const FOnQueryPartyJoinabilityReceivedDelegate& Delegate)
+{
+	auto DeprecationHelperLambda = [Delegate](const FUniqueNetId& LocalUserId, const FOnlinePartyId& PartyId, const IOnlinePartyPendingJoinRequestInfo& JoiningUsers)
+	{
+		TArray<IOnlinePartyUserPendingJoinRequestInfoConstRef> Users;
+		JoiningUsers.GetUsers(Users);
+		check(Users.Num() > 0);
+		Delegate.ExecuteIfBound(LocalUserId, PartyId, *Users[0]->GetUserId(), Users[0]->GetPlatform(), *Users[0]->GetJoinData());
+	};
+	return OnQueryPartyJoinabilityReceivedDelegates.Add(FOnQueryPartyJoinabilityGroupReceivedDelegate::CreateLambda(DeprecationHelperLambda));
+}
+
 bool FPartyConfiguration::operator==(const FPartyConfiguration& Other) const
 {
 	return JoinRequestAction == Other.JoinRequestAction &&
diff --git a/Engine/Plugins/Online/OnlineSubsystem/Source/Public/Interfaces/OnlinePartyInterface.h b/Engine/Plugins/Online/OnlineSubsystem/Source/Public/Interfaces/OnlinePartyInterface.h
--- a/Engine/Plugins/Online/OnlineSubsystem/Source/Public/Interfaces/OnlinePartyInterface.h
+++ b/Engine/Plugins/Online/OnlineSubsystem/Source/Public/Interfaces/OnlinePartyInterface.h
@@ -335,33 +335,104 @@ typedef TSharedRef<const FOnlinePartyData> FOnlinePartyDataConstRef;
 typedef TSharedPtr<const FOnlinePartyData> FOnlinePartyDataConstPtr;
 
 /**
-* Info needed to join a party
-*/
+ * Info about a user requesting to join a party
+ */
+class IOnlinePartyUserPendingJoinRequestInfo
+{
+public:
+	virtual ~IOnlinePartyUserPendingJoinRequestInfo() = default;
+
+	/**
+	 * Get the id of the user requesting to join
+	 * @return the id of the user requesting to join
+	 */
+	virtual const TSharedRef<const FUniqueNetId>& GetUserId() const = 0;
+
+	/**
+	 * Get the display name of the user requesting to join
+	 * @return the display name of the user requesting to join
+	 */
+	virtual const FString& GetDisplayName() const = 0;
+
+	/**
+	 * Get the platform of the user requesting to join
+	 * @return the platform of the user requesting to join
+	 */
+	virtual const FString& GetPlatform() const = 0;
+
+	/**
+	 * Get the join data of the user requesting to join
+	 * @return the join data of the user requesting to join
+	 */
+	virtual TSharedRef<const FOnlinePartyData> GetJoinData() const = 0;
+};
+
+typedef TSharedRef<const IOnlinePartyUserPendingJoinRequestInfo> IOnlinePartyUserPendingJoinRequestInfoConstRef;
+typedef TSharedPtr<const IOnlinePartyUserPendingJoinRequestInfo> IOnlinePartyUserPendingJoinRequestInfoConstPtr;
+
+/**
+ * Info about a group of users requesting to join a party
+ */
 class IOnlinePartyPendingJoinRequestInfo
+	: public TSharedFromThis<IOnlinePartyPendingJoinRequestInfo>
 {
 public:
-	IOnlinePartyPendingJoinRequestInfo() {}
-	virtual ~IOnlinePartyPendingJoinRequestInfo() {}
+	IOnlinePartyPendingJoinRequestInfo() = default;
+	virtual ~IOnlinePartyPendingJoinRequestInfo() = default;
 
 	/**
-	 * @return id of the sender of this join request
+	 * Get the primary user's id
+	 * @return id of the primary user of this join request
 	 */
-	virtual const TSharedRef<const FUniqueNetId>& GetSenderId() const = 0;
+	virtual const TSharedRef<const FUniqueNetId>& GetSenderId() const
+	{
+		TArray<IOnlinePartyUserPendingJoinRequestInfoConstRef> Users;
+		GetUsers(Users);
+		check(Users.Num() > 0);
+		return Users[0]->GetUserId();
+	}
+
+	/**
+	 * Get the primary user's display name
+	 * @return display name of the primary user of this join request
+	 */
+	virtual const FString& GetSenderDisplayName() const
+	{
+		TArray<IOnlinePartyUserPendingJoinRequestInfoConstRef> Users;
+		GetUsers(Users);
+		check(Users.Num() > 0);
+		return Users[0]->GetDisplayName();
+	}
 
 	/**
-	 * @return display name of the sender of this join request
+	 * Get the primary user's platform
+	 * @return platform of the primary user of this join request
 	 */
-	virtual const FString& GetSenderDisplayName() const = 0;
+	virtual const FString& GetSenderPlatform() const
+	{
+		TArray<IOnlinePartyUserPendingJoinRequestInfoConstRef> Users;
+		GetUsers(Users);
+		check(Users.Num() > 0);
+		return Users[0]->GetPlatform();
+	}
 
 	/**
-	 * @return platform of the sender of this join request
+	 * Get the primary user's join data
+	 * @return join data provided by the primary user for this join request
 	 */
-	virtual const FString& GetSenderPlatform() const = 0;
+	virtual TSharedRef<const FOnlinePartyData> GetSenderJoinData() const
+	{
+		TArray<IOnlinePartyUserPendingJoinRequestInfoConstRef> Users;
+		GetUsers(Users);
+		check(Users.Num() > 0);
+		return Users[0]->GetJoinData();
+	}
 
 	/**
-	 * @return join data provided by the sender for htis join request
+	 * Get the list of users requesting to join
+	 * @return array of users requesting to join
 	 */
-	virtual TSharedRef<const FOnlinePartyData> GetSenderJoinData() const = 0;
+	virtual void GetUsers(TArray<IOnlinePartyUserPendingJoinRequestInfoConstRef>& OutUsers) const = 0;
 };
 
 typedef TSharedRef<const IOnlinePartyPendingJoinRequestInfo> IOnlinePartyPendingJoinRequestInfoConstRef;
@@ -977,6 +1048,7 @@ PARTY_DECLARE_DELEGATETYPE(OnPartyInviteResponseReceived);
 
 /**
  * Notification when a new reservation request is received
+ * Deprecated - Use OnPartyGroupJoinRequestReceived
  * @param LocalUserId - id associated with this notification
  * @param PartyId - id associated with the party
  * @param SenderId - id of member that sent the request
@@ -987,18 +1059,28 @@ DECLARE_MULTICAST_DELEGATE_FiveParams(F_PREFIX(OnPartyJoinRequestReceived), cons
 PARTY_DECLARE_DELEGATETYPE(OnPartyJoinRequestReceived);
 
 /**
-* Notification when a new reservation request is received
-* @param LocalUserId - id associated with this notification
-* @param PartyId - id associated with the party
-* @param SenderId - id of member that sent the request
-* @param Platform - platform of member that sent the request
-* @param PartyData - data provided by the sender for the leader to use to determine joinability
-*/
+ * Notification when a new reservation request is received
+ * @param LocalUserId - id associated with this notification
+ * @param PartyId - id associated with the party
+ * @param JoinRequestInfo - data about users that are joining
+ */
+DECLARE_MULTICAST_DELEGATE_ThreeParams(F_PREFIX(OnPartyGroupJoinRequestReceived), const FUniqueNetId& /*LocalUserId*/, const FOnlinePartyId& /*PartyId*/, const IOnlinePartyPendingJoinRequestInfo& /*JoinRequestInfo*/);
+PARTY_DECLARE_DELEGATETYPE(OnPartyGroupJoinRequestReceived);
+
+/**
+ * Notification when a new reservation request is received
+ * @param LocalUserId - id associated with this notification
+ * @param PartyId - id associated with the party
+ * @param SenderId - id of member that sent the request
+ * @param Platform - platform of member that sent the request
+ * @param PartyData - data provided by the sender for the leader to use to determine joinability
+ */
 DECLARE_MULTICAST_DELEGATE_ThreeParams(F_PREFIX(OnPartyJIPRequestReceived), const FUniqueNetId& /*LocalUserId*/, const FOnlinePartyId& /*PartyId*/, const FUniqueNetId& /*SenderId*/);
 PARTY_DECLARE_DELEGATETYPE(OnPartyJIPRequestReceived);
 
 /**
  * Notification when a player wants to know if the party is in a joinable state
+ * Deprecated - Use OnQueryPartyJoinabilityGroupReceived
  * @param LocalUserId - id associated with this notification
  * @param PartyId - id associated with the party
  * @param SenderId - id of member that sent the request
@@ -1008,6 +1090,15 @@ PARTY_DECLARE_DELEGATETYPE(OnPartyJIPRequestReceived);
 DECLARE_MULTICAST_DELEGATE_FiveParams(F_PREFIX(OnQueryPartyJoinabilityReceived), const FUniqueNetId& /*LocalUserId*/, const FOnlinePartyId& /*PartyId*/, const FUniqueNetId& /*SenderId*/, const FString& /*Platform*/, const FOnlinePartyData& /*PartyData*/);
 PARTY_DECLARE_DELEGATETYPE(OnQueryPartyJoinabilityReceived);
 
+/**
+ * Notification when a player wants to know if the party is in a joinable state
+ * @param LocalUserId - id associated with this notification
+ * @param PartyId - id associated with the party
+ * @param JoinRequestInfo - data about users that are joining
+ */
+DECLARE_MULTICAST_DELEGATE_ThreeParams(F_PREFIX(OnQueryPartyJoinabilityGroupReceived), const FUniqueNetId& /*LocalUserId*/, const FOnlinePartyId& /*PartyId*/, const IOnlinePartyPendingJoinRequestInfo& /*JoinRequestInfo*/);
+PARTY_DECLARE_DELEGATETYPE(OnQueryPartyJoinabilityGroupReceived);
+
 /**
  * Request for the game to fill in data to be sent with the join request for the leader to make an informed decision based on the joiner's state
  * @param LocalUserId - id associated with this notification
@@ -1723,11 +1814,11 @@ public:
 	 * Subscriber is expected to call ApproveJoinRequest
 	 * @param LocalUserId - id associated with this notification
 	 * @param PartyId - id associated with the party
-	 * @param SenderId - id of member that sent the request
-	 * @param Platform - platform of member that sent the request
-	 * @param PartyData - data provided by the sender for the leader to use to determine joinability
+	 * @param JoinRequestInfo - data about users that are joining
 	 */
-	DEFINE_ONLINE_DELEGATE_FIVE_PARAM(OnPartyJoinRequestReceived, const FUniqueNetId& /*LocalUserId*/, const FOnlinePartyId& /*PartyId*/, const FUniqueNetId& /*SenderId*/, const FString& /*Platform*/, const FOnlinePartyData& /*PartyData*/);
+	DEFINE_ONLINE_DELEGATE_THREE_PARAM_DEPRECATION_HELPER(OnPartyJoinRequestReceived, OnPartyGroupJoinRequestReceived, const FUniqueNetId& /*LocalUserId*/, const FOnlinePartyId& /*PartyId*/, const IOnlinePartyPendingJoinRequestInfo& /*JoinRequestInfo*/);
+	UE_DEPRECATED(4.25, "Use OnPartyGroupJoinRequestReceived instead of OnPartyJoinRequestReceived")
+	virtual FDelegateHandle AddOnPartyJoinRequestReceivedDelegate_Handle(const FOnPartyJoinRequestReceivedDelegate& Delegate);
 
 	/**
 	* Notification when a new reservation request is received
@@ -1745,11 +1836,11 @@ public:
 	 * Subscriber is expected to call RespondToQueryJoinability
 	 * @param LocalUserId - id associated with this notification
 	 * @param PartyId - id associated with the party
-	 * @param SenderId - id of member that sent the request
-	 * @param Platform - platform of member that sent the request
-	 * @param PartyData - data provided by the sender for the leader to use to determine joinability
+	 * @param JoinRequestInfo - data about users that are joining
 	 */
-	DEFINE_ONLINE_DELEGATE_FIVE_PARAM(OnQueryPartyJoinabilityReceived, const FUniqueNetId& /*LocalUserId*/, const FOnlinePartyId& /*PartyId*/, const FUniqueNetId& /*SenderId*/, const FString& /*Platform*/, const FOnlinePartyData& /*PartyData*/);
+	DEFINE_ONLINE_DELEGATE_THREE_PARAM_DEPRECATION_HELPER(OnQueryPartyJoinabilityReceived, OnQueryPartyJoinabilityGroupReceived, const FUniqueNetId& /*LocalUserId*/, const FOnlinePartyId& /*PartyId*/, const IOnlinePartyPendingJoinRequestInfo& /*JoinRequestInfo*/);
+	UE_DEPRECATED(4.25, "Use OnQueryPartyJoinabilityGroupReceived instead of OnQueryPartyJoinabilityReceived")
+	virtual FDelegateHandle AddOnQueryPartyJoinabilityReceivedDelegate_Handle(const FOnQueryPartyJoinabilityReceivedDelegate& Delegate);
 
 	/**
 	 * Notification when a player wants to know if the party is in a joinable state
