#UE4 CIS fix

diff --git a/Engine/Plugins/Online/OnlineFramework/Source/Party/Public/SocialQuery.h b/Engine/Plugins/Online/OnlineFramework/Source/Party/Public/SocialQuery.h
--- a/Engine/Plugins/Online/OnlineFramework/Source/Party/Public/SocialQuery.h
+++ b/Engine/Plugins/Online/OnlineFramework/Source/Party/Public/SocialQuery.h
@@ -5,6 +5,7 @@
 #include "Containers/Ticker.h"
 #include "OnlineSubsystem.h"
 #include "SocialToolkit.h"
+#include "Misc/ConfigCacheIni.h"
 
 DECLARE_DELEGATE_TwoParams(FOnQueryCompleted, FName, const TSharedRef<class FSocialQueryBase>&);
 
