Fix not adding players to reservations array

diff --git a/Engine/Plugins/Online/OnlineFramework/Source/Party/Private/Party/SocialParty.cpp b/Engine/Plugins/Online/OnlineFramework/Source/Party/Private/Party/SocialParty.cpp
--- a/Engine/Plugins/Online/OnlineFramework/Source/Party/Private/Party/SocialParty.cpp
+++ b/Engine/Plugins/Online/OnlineFramework/Source/Party/Private/Party/SocialParty.cpp
@@ -1301,7 +1301,7 @@ void USocialParty::ConnectToReservationBeacon()
 							ReservationAsArray.Reserve(NextApproval.Members.Num());
 							for (const FPendingMemberApproval::FMemberInfo& MemberInfo : NextApproval.Members)
 							{
-								FPlayerReservation Reservation;
+								FPlayerReservation& Reservation = ReservationAsArray.Emplace_GetRef();
 								Reservation.UniqueId = MemberInfo.MemberId;
 								Reservation.Platform = MemberInfo.Platform;
 
@@ -1397,7 +1397,7 @@ void USocialParty::PumpApprovalQueue()
 			PlayersToAdd.Reserve(NextApproval.Members.Num());
 			for (const FPendingMemberApproval::FMemberInfo& MemberInfo : NextApproval.Members)
 			{
-				FPlayerReservation NewPlayerRes;
+				FPlayerReservation& NewPlayerRes = PlayersToAdd.Emplace_GetRef();
 				NewPlayerRes.UniqueId = MemberInfo.MemberId;
 				NewPlayerRes.Platform = MemberInfo.Platform;
 
@@ -1413,7 +1413,6 @@ void USocialParty::PumpApprovalQueue()
 					// This doesn't matter, since the crossplay state of the match has already been set.
 					NewPlayerRes.bAllowCrossplay = true;
 				}
-				PlayersToAdd.Emplace(MoveTemp(NewPlayerRes));
 			}
 			ReservationBeaconClient->RequestReservationUpdate(GetPartyLeader()->GetPrimaryNetId(), PlayersToAdd);
 		}
