#online Adjustments to handle secondary player leaving a party and being able to rejoin again later. For the most part this is ensuring that LeaveParty has the UserId of the player leaving the party.
The additional MCP Party System changes have been passed to, and are shelved in another review from Chris.Applegarth.

diff --git a/Engine/Plugins/Online/OnlineFramework/Source/Party/Private/Party/SocialParty.cpp b/Engine/Plugins/Online/OnlineFramework/Source/Party/Private/Party/SocialParty.cpp
--- a/Engine/Plugins/Online/OnlineFramework/Source/Party/Private/Party/SocialParty.cpp
+++ b/Engine/Plugins/Online/OnlineFramework/Source/Party/Private/Party/SocialParty.cpp
@@ -912,6 +912,25 @@ void USocialParty::HandleLeavePartyComplete(const FUniqueNetId& LocalUserId, con
 	OnAttemptComplete.ExecuteIfBound(LeaveResult);
 }
 
+void USocialParty::HandleRemoveLocalPlayerComplete(const FUniqueNetId& LocalUserId, const FOnlinePartyId& PartyId, ELeavePartyCompletionResult LeaveResult, FOnLeavePartyAttemptComplete OnAttemptComplete)
+{
+	UE_LOG(LogParty, Verbose, TEXT("Removed Local player [%s] is no longer in party [%s]."), *LocalUserId.ToDebugString(), *ToDebugString());
+
+	// Remove the secondary member from the party members ids
+	for (UPartyMember* PartyMember : GetPartyMembers())
+	{
+		if (LocalUserId == *PartyMember->GetPrimaryNetId())
+		{
+			PartyMember->NotifyRemovedFromParty(EMemberExitedReason::Unknown);
+			PartyMember->MarkPendingKill();
+			PartyMembersById.Remove(PartyMember->GetPrimaryNetId());
+			break;
+		}
+	}
+
+	OnAttemptComplete.ExecuteIfBound(LeaveResult);
+}
+
 void USocialParty::HandlePrivacySettingsChanged(const FPartyPrivacySettings& NewPrivacySettings)
 {
 	check(IsLocalPlayerPartyLeader());
@@ -959,9 +978,17 @@ void USocialParty::HandlePartyLeft(const FUniqueNetId& LocalUserId, const FOnlin
 	// this function is called when a party is left due to unintentional leave (e.g. disconnect)
 	if (PartyId == GetPartyId())
 	{
-		// process an full "leave" for the party which will clean it up here and in OnlinePartyMcp
-		// this will also trigger a new persistent party to be created
-		LeaveParty();
+		if (LocalUserId == *OwningLocalUserId)
+		{
+			// process an full "leave" for the party which will clean it up here and in OnlinePartyMcp
+			// this will also trigger a new persistent party to be created
+			LeaveParty();
+		}
+		else
+		{
+			// process just the secondary member leaving the party
+			RemoveLocalMember(FUniqueNetIdRepl(LocalUserId.AsShared()));
+		}
 	}
 }
 
@@ -976,6 +1003,7 @@ void USocialParty::HandlePartyMemberExited(const FUniqueNetId& LocalUserId, cons
 				//@todo DanH Party: Do I get this for a self-initiated party leave as well? #required
 				if (!bIsLeavingParty)
 				{
+					check(LocalUserId == *OwningLocalUserId);		// Only the primary player is allowed to finalize a party leave
 					FinalizePartyLeave(ExitReason);
 				}
 			}
@@ -1183,6 +1211,7 @@ void USocialParty::LeaveParty(const FOnLeavePartyAttemptComplete& OnLeaveAttempt
 		const FOnlinePartyId& PartyId = GetPartyId();
 		if (OwningLocalUserId.IsValid() && PartyId.IsValid())
 		{
+			// All local players will be removed as a consequence of leaving the party with the primary player
 			const IOnlinePartyPtr PartyInterface = Online::GetPartyInterfaceChecked(GetWorld());
 			FOnLeavePartyComplete OnLeaveComplete = FOnLeavePartyComplete::CreateUObject(this, &USocialParty::HandleLeavePartyComplete, OnLeaveAttemptComplete);
 			PartyInterface->LeaveParty(*OwningLocalUserId, PartyId, OnLeaveComplete);
@@ -1194,6 +1223,31 @@ void USocialParty::LeaveParty(const FOnLeavePartyAttemptComplete& OnLeaveAttempt
 	}
 }
 
+void USocialParty::RemoveLocalMember(const FUniqueNetIdRepl& LocalUserId, const FOnLeavePartyAttemptComplete& OnLeaveAttemptComplete)
+{
+	if (bIsLeavingParty)
+	{
+		// Already working on it for the primary player!
+		OnLeaveAttemptComplete.ExecuteIfBound(ELeavePartyCompletionResult::LeavePending);
+	}
+	else
+	{
+		UE_LOG(LogParty, Verbose, TEXT("Attempting to leave party LocalUserId=[%s] [%s] "), *LocalUserId.ToDebugString(), *ToDebugString());
+
+		const FOnlinePartyId& PartyId = GetPartyId();
+		if (LocalUserId.IsValid() && PartyId.IsValid())
+		{
+			const IOnlinePartyPtr PartyInterface = Online::GetPartyInterfaceChecked(GetWorld());
+			FOnLeavePartyComplete OnLeaveComplete = FOnLeavePartyComplete::CreateUObject(this, &USocialParty::HandleRemoveLocalPlayerComplete, OnLeaveAttemptComplete);
+			PartyInterface->LeaveParty(*LocalUserId, PartyId, OnLeaveComplete);
+		}
+		else
+		{
+			OnLeaveAttemptComplete.ExecuteIfBound(ELeavePartyCompletionResult::UnknownClientFailure);
+		}
+	}
+}
+
 bool USocialParty::ContainsUser(const USocialUser& User) const
 {
 	for(const UPartyMember* PartyMember : GetPartyMembers())
diff --git a/Engine/Plugins/Online/OnlineFramework/Source/Party/Public/Party/SocialParty.h b/Engine/Plugins/Online/OnlineFramework/Source/Party/Public/Party/SocialParty.h
--- a/Engine/Plugins/Online/OnlineFramework/Source/Party/Public/Party/SocialParty.h
+++ b/Engine/Plugins/Online/OnlineFramework/Source/Party/Public/Party/SocialParty.h
@@ -84,6 +84,7 @@ public:
 
 	DECLARE_DELEGATE_OneParam(FOnLeavePartyAttemptComplete, ELeavePartyCompletionResult)
 	virtual void LeaveParty(const FOnLeavePartyAttemptComplete& OnLeaveAttemptComplete = FOnLeavePartyAttemptComplete());
+	virtual void RemoveLocalMember(const FUniqueNetIdRepl& LocalUserId, const FOnLeavePartyAttemptComplete& OnLeaveAttemptComplete = FOnLeavePartyAttemptComplete());
 
 	const FPartyRepData& GetRepData() const { return *PartyDataReplicator; }
 
@@ -325,6 +326,7 @@ private:	// Handlers
 	void HandleReservationRequestComplete(EPartyReservationResult::Type ReservationResponse);
 
 	void HandleLeavePartyComplete(const FUniqueNetId& LocalUserId, const FOnlinePartyId& PartyId, ELeavePartyCompletionResult LeaveResult, FOnLeavePartyAttemptComplete OnAttemptComplete);
+	void HandleRemoveLocalPlayerComplete(const FUniqueNetId& LocalUserId, const FOnlinePartyId& PartyId, ELeavePartyCompletionResult LeaveResult, FOnLeavePartyAttemptComplete OnAttemptComplete);
 
 	void RemovePlayerFromReservationBeacon(const FUniqueNetId& LocalUserId, const FUniqueNetId& PlayerToRemove);
 private:
