[Leto] Prevent secondary players from accessing 'Select Game Mode' or or 'Leave Party / Leave Match' in the Main Menu blade. Also prevent "leave party interaction" on the lobby widgets. P2 should stay with their primary player unless they log out.

P1 can still see the 'Select Game Mode' button, but if they click it, they will get a confirmation prompting them to log out the secondary player before switching game modes.

diff --git a/Engine/Plugins/Online/OnlineFramework/Source/Party/Private/Interactions/PartyInteractions.cpp b/Engine/Plugins/Online/OnlineFramework/Source/Party/Private/Interactions/PartyInteractions.cpp
--- a/Engine/Plugins/Online/OnlineFramework/Source/Party/Private/Interactions/PartyInteractions.cpp
+++ b/Engine/Plugins/Online/OnlineFramework/Source/Party/Private/Interactions/PartyInteractions.cpp
@@ -126,6 +126,13 @@ bool FSocialInteraction_LeaveParty::CanExecute(const USocialUser& User)
 {
 	if (User.IsLocalUser())
 	{
+		USocialToolkit& OwningToolkit = User.GetOwningToolkit();
+		ULocalPlayer& LocalPlayer = OwningToolkit.GetOwningLocalPlayer();
+		if (!LocalPlayer.IsPrimaryPlayer())
+		{
+			return false;
+		}
+
 		const UPartyMember* LocalMember = User.GetPartyMember(IOnlinePartySystem::GetPrimaryPartyTypeId());
 		return LocalMember && LocalMember->GetParty().GetNumPartyMembers() > 1;
 	}
