Non-unity include error (LocalPlayer for PartyInteractions.cpp)

diff --git a/Engine/Plugins/Online/OnlineFramework/Source/Party/Private/Interactions/PartyInteractions.cpp b/Engine/Plugins/Online/OnlineFramework/Source/Party/Private/Interactions/PartyInteractions.cpp
--- a/Engine/Plugins/Online/OnlineFramework/Source/Party/Private/Interactions/PartyInteractions.cpp
+++ b/Engine/Plugins/Online/OnlineFramework/Source/Party/Private/Interactions/PartyInteractions.cpp
@@ -7,6 +7,7 @@
 #include "Party/SocialParty.h"
 #include "Party/PartyMember.h" 
 #include "Interfaces/OnlinePartyInterface.h"
+#include "Engine/LocalPlayer.h"
 
 #define LOCTEXT_NAMESPACE "PartyInteractions"
 
