Making sure SocialToolkits are created in the same order as LocalPlayers, and using player index to delete SocialToolkits in the event of a player removal.

diff --git a/Engine/Plugins/Online/OnlineFramework/Source/Party/Private/SocialManager.cpp b/Engine/Plugins/Online/OnlineFramework/Source/Party/Private/SocialManager.cpp
--- a/Engine/Plugins/Online/OnlineFramework/Source/Party/Private/SocialManager.cpp
+++ b/Engine/Plugins/Online/OnlineFramework/Source/Party/Private/SocialManager.cpp
@@ -659,7 +659,7 @@ UGameInstance& USocialManager::GetGameInstance() const
 	return *GetTypedOuter<UGameInstance>();
 }
 
-USocialToolkit& USocialManager::CreateSocialToolkit(ULocalPlayer& OwningLocalPlayer)
+USocialToolkit& USocialManager::CreateSocialToolkit(ULocalPlayer& OwningLocalPlayer, int32 LocalPlayerIndex)
 {
 	for (USocialToolkit* ExistingToolkit : SocialToolkits)
 	{
@@ -668,7 +668,7 @@ USocialToolkit& USocialManager::CreateSocialToolkit(ULocalPlayer& OwningLocalPla
 	check(ToolkitClass);
 
 	USocialToolkit* NewToolkit = NewObject<USocialToolkit>(this, ToolkitClass);
-	SocialToolkits.Add(NewToolkit);
+	SocialToolkits.Insert(NewToolkit, LocalPlayerIndex);
 	NewToolkit->InitializeToolkit(OwningLocalPlayer);
 	OnToolkitCreatedInternal(*NewToolkit);
 	NewToolkit->OnToolkitReset().AddUObject(this, &USocialManager::HandleToolkitReset, NewToolkit->GetLocalUserNum());
@@ -865,9 +865,12 @@ void USocialManager::HandleGameViewportInitialized()
 	GameViewport->OnPlayerRemoved().AddUObject(this, &USocialManager::HandleLocalPlayerRemoved);
 
 	// Immediately spin up toolkits for local players that already exist
-	for (ULocalPlayer* ExistingLocalPlayer : GameInstance.GetLocalPlayers())
+	const TArray<ULocalPlayer*>& LocalPlayers = GameInstance.GetLocalPlayers();
+	for (int32 LocalPlayerIndex = 0; LocalPlayerIndex < LocalPlayers.Num(); ++LocalPlayerIndex)
 	{
-		CreateSocialToolkit(*ExistingLocalPlayer);
+		ULocalPlayer* ExistingLocalPlayer = LocalPlayers[LocalPlayerIndex];
+
+		CreateSocialToolkit(*ExistingLocalPlayer, LocalPlayerIndex);
 	}
 }
 
@@ -895,15 +898,19 @@ void USocialManager::HandleLocalPlayerAdded(int32 LocalUserNum)
 	ULocalPlayer* NewLocalPlayer = GetGameInstance().GetLocalPlayerByIndex(LocalUserNum);
 	check(NewLocalPlayer);
 
-	CreateSocialToolkit(*NewLocalPlayer);
+	CreateSocialToolkit(*NewLocalPlayer, LocalUserNum);
 }
 
 void USocialManager::HandleLocalPlayerRemoved(int32 LocalUserNum)
 {
-	if (USocialToolkit* Toolkit = GetSocialToolkit(LocalUserNum))
+	//GetSocialToolkit accepts a ControllerId, not a player index, so we'll access it directly 
+	if (SocialToolkits.IsValidIndex(LocalUserNum))
 	{
-		SocialToolkits.Remove(Toolkit);
-		Toolkit->MarkPendingKill();
+		if (USocialToolkit* Toolkit = SocialToolkits[LocalUserNum])
+		{
+			SocialToolkits.Remove(Toolkit);
+			Toolkit->MarkPendingKill();
+		}
 	}
 }
 
diff --git a/Engine/Plugins/Online/OnlineFramework/Source/Party/Public/SocialManager.h b/Engine/Plugins/Online/OnlineFramework/Source/Party/Public/SocialManager.h
--- a/Engine/Plugins/Online/OnlineFramework/Source/Party/Public/SocialManager.h
+++ b/Engine/Plugins/Online/OnlineFramework/Source/Party/Public/SocialManager.h
@@ -200,7 +200,7 @@ protected:
 
 private:
 	UGameInstance& GetGameInstance() const;
-	USocialToolkit& CreateSocialToolkit(ULocalPlayer& OwningLocalPlayer);
+	USocialToolkit& CreateSocialToolkit(ULocalPlayer& OwningLocalPlayer, int32 LocalPlayerIndex);
 
 	void QueryPartyJoinabilityInternal(FJoinPartyAttempt& JoinAttempt);
 	void JoinPartyInternal(FJoinPartyAttempt& JoinAttempt);
