Added 'On<X>Complete' stubs in the SocialToolkit and called them from within their respective 'Handle<X>Complete' handlers. The game can now react to these notifications while preserving a robust API.

#tests Compile DebugGame Editor Win64 / Development Editor Win64 / Shipping Client / Ran it in PIE and verified flow.

diff --git a/Engine/Plugins/Online/OnlineFramework/Source/Party/Private/SocialToolkit.cpp b/Engine/Plugins/Online/OnlineFramework/Source/Party/Private/SocialToolkit.cpp
--- a/Engine/Plugins/Online/OnlineFramework/Source/Party/Private/SocialToolkit.cpp
+++ b/Engine/Plugins/Online/OnlineFramework/Source/Party/Private/SocialToolkit.cpp
@@ -691,6 +691,8 @@ void USocialToolkit::HandleReadFriendsListComplete(int32 LocalUserNum, bool bWas
 	}
 
 	HandleExistingPartyInvites(SubsystemType);
+
+	OnReadFriendsListComplete(LocalUserNum, bWasSuccessful, ListName, ErrorStr, SubsystemType);
 }
 
 void USocialToolkit::HandleQueryBlockedPlayersComplete(const FUniqueNetId& UserId, bool bWasSuccessful, const FString& ErrorStr, ESocialSubsystem SubsystemType)
@@ -714,6 +716,8 @@ void USocialToolkit::HandleQueryBlockedPlayersComplete(const FUniqueNetId& UserI
 			//@todo DanH: Only bother retrying on primary
 		}
 	}
+
+	OnQueryBlockedPlayersComplete(UserId, bWasSuccessful, ErrorStr, SubsystemType);
 }
 
 void USocialToolkit::HandleQueryRecentPlayersComplete(const FUniqueNetId& UserId, const FString& Namespace, bool bWasSuccessful, const FString& ErrorStr, ESocialSubsystem SubsystemType)
@@ -737,6 +741,8 @@ void USocialToolkit::HandleQueryRecentPlayersComplete(const FUniqueNetId& UserId
 			//@todo DanH: Only bother retrying on primary
 		}
 	}
+
+	OnQueryRecentPlayersComplete(UserId, Namespace, bWasSuccessful, ErrorStr, SubsystemType);
 }
 
 void USocialToolkit::HandleRecentPlayersAdded(const FUniqueNetId& LocalUserId, const TArray<TSharedRef<FOnlineRecentPlayer>>& NewRecentPlayers, ESocialSubsystem SubsystemType)
@@ -913,6 +919,8 @@ void USocialToolkit::HandleDeleteFriendComplete(int32 InLocalUserNum, bool bWasS
 			FormerFriend->NotifyUserUnfriended(SubsystemType);
 		}
 	}
+
+	OnDeleteFriendComplete(InLocalUserNum, bWasSuccessful, DeletedFriendId, ListName, ErrorStr, SubsystemType);
 }
 
 void USocialToolkit::HandleAcceptFriendInviteComplete(int32 LocalUserNum, bool bWasSuccessful, const FUniqueNetId& InviterUserId, const FString& ListName, const FString& ErrorStr)
@@ -964,6 +972,8 @@ void USocialToolkit::HandleBlockPlayerComplete(int32 LocalUserNum, bool bWasSucc
 				}
 			});
 	}
+
+	OnBlockPlayerComplete(LocalUserNum, bWasSuccessful, BlockedPlayerId, ListName, ErrorStr, SubsystemType);
 }
 
 void USocialToolkit::HandleUnblockPlayerComplete(int32 LocalUserNum, bool bWasSuccessful, const FUniqueNetId& UnblockedPlayerId, const FString& ListName, const FString& ErrorStr, ESocialSubsystem SubsystemType)
@@ -976,6 +986,8 @@ void USocialToolkit::HandleUnblockPlayerComplete(int32 LocalUserNum, bool bWasSu
 			UnblockedUser->NotifyUserUnblocked(SubsystemType);
 		}
 	}
+
+	OnUnblockPlayerComplete(LocalUserNum, bWasSuccessful, UnblockedPlayerId, ListName, ErrorStr, SubsystemType);
 }
 
 //@todo DanH recent players: Where is the line for this between backend and game to update this stuff? #required
diff --git a/Engine/Plugins/Online/OnlineFramework/Source/Party/Public/SocialToolkit.h b/Engine/Plugins/Online/OnlineFramework/Source/Party/Public/SocialToolkit.h
--- a/Engine/Plugins/Online/OnlineFramework/Source/Party/Public/SocialToolkit.h
+++ b/Engine/Plugins/Online/OnlineFramework/Source/Party/Public/SocialToolkit.h
@@ -132,6 +132,12 @@ protected:
 	
 	virtual void OnSendFriendInviteComplete(const FUniqueNetId& InvitedUserId, const FString& InviteeName, bool bInviteSuceeded, const FString& InviteError) {}
 	virtual void OnAcceptFriendInviteComplete(const FUniqueNetId& InviterUserId, bool bInviteSucceeded, const FString& InviteError) {}
+	virtual void OnReadFriendsListComplete(int32 LocalUserNum, bool bWasSuccessful, const FString& ListName, const FString& ErrorStr, ESocialSubsystem SubsystemType) {}
+	virtual void OnQueryBlockedPlayersComplete(const FUniqueNetId& UserId, bool bWasSuccessful, const FString& ErrorStr, ESocialSubsystem SubsystemType) {}
+	virtual void OnQueryRecentPlayersComplete(const FUniqueNetId& UserId, const FString& Namespace, bool bWasSuccessful, const FString& ErrorStr, ESocialSubsystem SubsystemType) {}
+	virtual void OnDeleteFriendComplete(int32 LocalPlayer, bool bWasSuccessful, const FUniqueNetId& FormerFriendId, const FString& ListName, const FString& ErrorStr, ESocialSubsystem SubsystemType) {}
+	virtual void OnBlockPlayerComplete(int32 LocalUserNum, bool bWasSuccessful, const FUniqueNetId& BlockedPlayerID, const FString& ListName, const FString& ErrorStr, ESocialSubsystem SubsystemType) {}
+	virtual void OnUnblockPlayerComplete(int32 LocalUserNum, bool bWasSuccessful, const FUniqueNetId& UnblockedPlayerID, const FString& ListName, const FString& ErrorStr, ESocialSubsystem SubsystemType) {}
 
 	void QueryFriendsLists();
 	void QueryBlockedPlayers();
