Added missed clearing of delegates on the FriendsInterface in USocialToolkit::OnOwnerLoggedOut()

diff --git a/Engine/Plugins/Online/OnlineFramework/Source/Party/Private/SocialToolkit.cpp b/Engine/Plugins/Online/OnlineFramework/Source/Party/Private/SocialToolkit.cpp
--- a/Engine/Plugins/Online/OnlineFramework/Source/Party/Private/SocialToolkit.cpp
+++ b/Engine/Plugins/Online/OnlineFramework/Source/Party/Private/SocialToolkit.cpp
@@ -540,10 +540,13 @@ void USocialToolkit::OnOwnerLoggedOut()
 				FriendsInterface->ClearOnInviteReceivedDelegates(this);
 				FriendsInterface->ClearOnInviteAcceptedDelegates(this);
 				FriendsInterface->ClearOnInviteRejectedDelegates(this);
+				FriendsInterface->ClearOnInviteAbortedDelegates(this);
 
 				FriendsInterface->ClearOnBlockedPlayerCompleteDelegates(LocalUserNum, this);
 				FriendsInterface->ClearOnUnblockedPlayerCompleteDelegates(LocalUserNum, this);
 
+				FriendsInterface->ClearOnRecentPlayersAddedDelegates(this);
+
 				FriendsInterface->ClearOnQueryBlockedPlayersCompleteDelegates(this);
 				FriendsInterface->ClearOnQueryRecentPlayersCompleteDelegates(this);
 			}
