#OnlineFramework Fix crash where HotfixManager is created early enough to be in the disregardforgc window and thus cannot reference non-root objects by making the objects it references root

diff --git a/Engine/Plugins/Online/OnlineFramework/Source/Hotfix/Private/OnlineHotfixManager.cpp b/Engine/Plugins/Online/OnlineFramework/Source/Hotfix/Private/OnlineHotfixManager.cpp
--- a/Engine/Plugins/Online/OnlineFramework/Source/Hotfix/Private/OnlineHotfixManager.cpp
+++ b/Engine/Plugins/Online/OnlineFramework/Source/Hotfix/Private/OnlineHotfixManager.cpp
@@ -1459,6 +1459,13 @@ void UOnlineHotfixManager::PatchAssetsFromIniFiles()
 							// this session, so we reference it to prevent it from being evicted from memory.  It's OK if we end up re-patching
 							// the same asset multiple times per session.
 							AssetsHotfixedFromIniFiles.Add(Asset);
+							
+							// If the hotfixmanager was made before the DisregardForGC window closed, these references will always remain in memory,
+							// but this object wont be considered for GC, so we'll need to add the referenced object to root to prevent GC.
+							if (GUObjectArray.IsDisregardForGC(this))
+							{
+								Asset->AddToRoot();
+							}
 						}
 					}
 					else
