Back out changelist 10844330. The root cause that necessitated this was found, so I'm removing this unncessary change

diff --git a/Engine/Plugins/Online/OnlineFramework/Source/Hotfix/Private/OnlineHotfixManager.cpp b/Engine/Plugins/Online/OnlineFramework/Source/Hotfix/Private/OnlineHotfixManager.cpp
--- a/Engine/Plugins/Online/OnlineFramework/Source/Hotfix/Private/OnlineHotfixManager.cpp
+++ b/Engine/Plugins/Online/OnlineFramework/Source/Hotfix/Private/OnlineHotfixManager.cpp
@@ -1459,13 +1459,6 @@ void UOnlineHotfixManager::PatchAssetsFromIniFiles()
 							// this session, so we reference it to prevent it from being evicted from memory.  It's OK if we end up re-patching
 							// the same asset multiple times per session.
 							AssetsHotfixedFromIniFiles.Add(Asset);
-							
-							// If the hotfixmanager was made before the DisregardForGC window closed, these references will always remain in memory,
-							// but this object wont be considered for GC, so we'll need to add the referenced object to root to prevent GC.
-							if (GUObjectArray.IsDisregardForGC(this))
-							{
-								Asset->AddToRoot();
-							}
 						}
 					}
 					else
