Removed use of implicit conversions from FUniqueNetId to FUniqueNetIdWrapper in preparation for disabling those conversions, also removing some needless back and forth conversions.

diff --git a/Engine/Plugins/Online/OnlineFramework/Source/Party/Private/Chat/SocialGroupChannel.cpp b/Engine/Plugins/Online/OnlineFramework/Source/Party/Private/Chat/SocialGroupChannel.cpp
--- a/Engine/Plugins/Online/OnlineFramework/Source/Party/Private/Chat/SocialGroupChannel.cpp
+++ b/Engine/Plugins/Online/OnlineFramework/Source/Party/Private/Chat/SocialGroupChannel.cpp
@@ -11,7 +11,7 @@ USocialGroupChannel::USocialGroupChannel()
 
 void USocialGroupChannel::Initialize(IOnlineGroupsPtr InGroupInterface, USocialUser& InSocialUser, const FUniqueNetId& InGroupId)
 {
-	GroupId = InGroupId;
+	GroupId = InGroupId.AsShared();
 	SocialUser = &InSocialUser;
 	GroupInterfacePtr = InGroupInterface;
 
diff --git a/Engine/Plugins/Online/OnlineFramework/Source/Party/Private/Party/SocialParty.cpp b/Engine/Plugins/Online/OnlineFramework/Source/Party/Private/Party/SocialParty.cpp
--- a/Engine/Plugins/Online/OnlineFramework/Source/Party/Private/Party/SocialParty.cpp
+++ b/Engine/Plugins/Online/OnlineFramework/Source/Party/Private/Party/SocialParty.cpp
@@ -575,7 +575,7 @@ UPartyMember* USocialParty::GetOrCreatePartyMember(const FUniqueNetId& MemberId)
 		}
 		else
 		{
-			const bool bIsLocalUser = GetSocialManager().IsLocalUser(MemberId, ESocialSubsystem::Primary);
+			const bool bIsLocalUser = GetSocialManager().IsLocalUser(MemberId.AsShared(), ESocialSubsystem::Primary);
 			TSubclassOf<UPartyMember> PartyMemberClass = GetDesiredMemberClass(bIsLocalUser);
 			if (ensure(PartyMemberClass))
 			{
@@ -1011,7 +1011,7 @@ void USocialParty::HandlePartyMemberExited(const FUniqueNetId& LocalUserId, cons
 {
 	if (PartyId == GetPartyId())
 	{
-		if (UPartyMember** FoundPartyMember = PartyMembersById.Find(MemberId))
+		if (UPartyMember** FoundPartyMember = PartyMembersById.Find(MemberId.AsShared()))
 		{
 			if (LocalUserId == MemberId)
 			{
diff --git a/Engine/Plugins/Online/OnlineFramework/Source/Party/Private/SocialToolkit.cpp b/Engine/Plugins/Online/OnlineFramework/Source/Party/Private/SocialToolkit.cpp
--- a/Engine/Plugins/Online/OnlineFramework/Source/Party/Private/SocialToolkit.cpp
+++ b/Engine/Plugins/Online/OnlineFramework/Source/Party/Private/SocialToolkit.cpp
@@ -785,7 +785,7 @@ void USocialToolkit::HandleMapExternalIdComplete(ESocialSubsystem SubsystemType,
 void USocialToolkit::HandlePresenceReceived(const FUniqueNetId& UserId, const TSharedRef<FOnlineUserPresence>& NewPresence, ESocialSubsystem SubsystemType)
 {
 	QUICK_SCOPE_CYCLE_COUNTER(STAT_USocialToolkit_HandlePresenceReceived);
-	if (USocialUser* UpdatedUser = FindUser(UserId))
+	if (USocialUser* UpdatedUser = FindUser(UserId.AsShared()))
 	{
 		UpdatedUser->NotifyPresenceChanged(SubsystemType);
 	}
@@ -807,7 +807,7 @@ void USocialToolkit::HandleQueryPrimaryUserIdMappingComplete(bool bWasSuccessful
 	}
 	else
 	{
-		QueueUserDependentActionInternal(IdentifiedUserId, ESocialSubsystem::Primary,
+		QueueUserDependentActionInternal(IdentifiedUserId.AsShared(), ESocialSubsystem::Primary,
 			[this, DisplayName] (USocialUser& SocialUser)
 			{
 				if (SocialUser.IsBlocked())
@@ -830,7 +830,7 @@ void USocialToolkit::HandleFriendInviteReceived(const FUniqueNetId& LocalUserId,
 {
 	if (LocalUserId == *GetLocalUserNetId(SubsystemType))
 	{
-		QueueUserDependentActionInternal(SenderId, SubsystemType,
+		QueueUserDependentActionInternal(SenderId.AsShared(), SubsystemType,
 			[this, SubsystemType] (USocialUser& SocialUser)
 			{
 				//@todo DanH: This event should send the name of the list the accepting friend is on, shouldn't it?
@@ -851,7 +851,7 @@ void USocialToolkit::HandleFriendInviteAccepted(const FUniqueNetId& LocalUserId,
 {
 	if (LocalUserId == *GetLocalUserNetId(SubsystemType))
 	{
-		QueueUserDependentActionInternal(FriendId, SubsystemType,
+		QueueUserDependentActionInternal(FriendId.AsShared(), SubsystemType,
 			[this, SubsystemType] (USocialUser& SocialUser)
 			{
 				//@todo DanH: This event should send the name of the list the accepting friend is on, shouldn't it?
@@ -872,7 +872,7 @@ void USocialToolkit::HandleFriendInviteRejected(const FUniqueNetId& LocalUserId,
 {
 	if (LocalUserId == *GetLocalUserNetId(SubsystemType))
 	{
-		if (USocialUser* InvitedUser = FindUser(FriendId))
+		if (USocialUser* InvitedUser = FindUser(FriendId.AsShared()))
 		{
 			InvitedUser->NotifyFriendInviteRemoved(SubsystemType);
 		}
@@ -883,7 +883,7 @@ void USocialToolkit::HandleFriendInviteSent(int32 LocalUserNum, bool bWasSuccess
 {
 	if (bWasSuccessful)
 	{
-		QueueUserDependentActionInternal(InvitedUserId, SubsystemType,
+		QueueUserDependentActionInternal(InvitedUserId.AsShared(), SubsystemType,
 			[this, SubsystemType, ListName, LocalUserNum] (USocialUser& SocialUser)
 			{
 				IOnlineFriendsPtr FriendsInterface = Online::GetFriendsInterfaceChecked(GetWorld(), USocialManager::GetSocialOssName(SubsystemType));
@@ -904,7 +904,7 @@ void USocialToolkit::HandleFriendRemoved(const FUniqueNetId& LocalUserId, const
 {
 	if (LocalUserId == *GetLocalUserNetId(SubsystemType))
 	{
-		USocialUser* FormerFriend = FindUser(FriendId);
+		USocialUser* FormerFriend = FindUser(FriendId.AsShared());
 		if (ensure(FormerFriend))
 		{
 			FormerFriend->NotifyUserUnfriended(SubsystemType);
@@ -916,7 +916,7 @@ void USocialToolkit::HandleDeleteFriendComplete(int32 InLocalUserNum, bool bWasS
 {
 	if (bWasSuccessful && InLocalUserNum == GetLocalUserNum())
 	{
-		USocialUser* FormerFriend = FindUser(DeletedFriendId);
+		USocialUser* FormerFriend = FindUser(DeletedFriendId.AsShared());
 		if (ensure(FormerFriend))
 		{
 			FormerFriend->NotifyUserUnfriended(SubsystemType);
@@ -936,7 +936,7 @@ void USocialToolkit::HandlePartyInviteReceived(const FUniqueNetId& LocalUserId,
 	if (LocalUserId == *GetLocalUserNetId(ESocialSubsystem::Primary))
 	{
 		// We really should know about the sender of the invite already, but queue it up in case we receive it during initial setup
-		QueueUserDependentActionInternal(SenderId, ESocialSubsystem::Primary,
+		QueueUserDependentActionInternal(SenderId.AsShared(), ESocialSubsystem::Primary,
 			[this] (USocialUser& User)
 			{
 				if (User.IsFriend(ESocialSubsystem::Primary))
@@ -951,7 +951,7 @@ void USocialToolkit::HandleBlockPlayerComplete(int32 LocalUserNum, bool bWasSucc
 {
 	if (bWasSuccessful && LocalUserNum == GetLocalUserNum())
 	{
-		QueueUserDependentActionInternal(BlockedPlayerId, SubsystemType, 
+		QueueUserDependentActionInternal(BlockedPlayerId.AsShared(), SubsystemType, 
 			[this, SubsystemType] (USocialUser& User)
 			{
 				// Quite frustrating that the event doesn't sent the FOnlineBlockedPlayer in the first place or provide a direct getter on the interface...
@@ -983,7 +983,7 @@ void USocialToolkit::HandleUnblockPlayerComplete(int32 LocalUserNum, bool bWasSu
 {
 	if (bWasSuccessful && LocalUserNum == GetLocalUserNum())
 	{
-		USocialUser* UnblockedUser = FindUser(UnblockedPlayerId);
+		USocialUser* UnblockedUser = FindUser(UnblockedPlayerId.AsShared());
 		if (ensure(UnblockedUser))
 		{
 			UnblockedUser->NotifyUserUnblocked(SubsystemType);
diff --git a/Engine/Plugins/Online/OnlineFramework/Source/Party/Public/Chat/SocialChatManager.cpp b/Engine/Plugins/Online/OnlineFramework/Source/Party/Public/Chat/SocialChatManager.cpp
--- a/Engine/Plugins/Online/OnlineFramework/Source/Party/Public/Chat/SocialChatManager.cpp
+++ b/Engine/Plugins/Online/OnlineFramework/Source/Party/Public/Chat/SocialChatManager.cpp
@@ -558,7 +558,6 @@ void USocialChatManager::RefreshGroupsRequestCompleted(FGroupsResult Result)
 
 void USocialChatManager::OnGroupUpdated(const FUniqueNetId& GroupId)
 {
-	printf("");
 }
 
 bool USocialChatManager::IsUniqueIdOfOwner(const FUniqueNetId& LocalUserId) const
@@ -595,7 +594,7 @@ USocialGroupChannel& USocialChatManager::FindOrCreateGroupChannel(IOnlineGroupsP
 	USocialGroupChannel* NewGroupChannel = NewObject<USocialGroupChannel>(this, NewGroupClass);
 	check(NewGroupChannel);
 
-	GroupChannels.Add(FUniqueNetIdRepl(GroupId), NewGroupChannel);
+	GroupChannels.Add(GroupId.AsShared(), NewGroupChannel);
 
 	NewGroupChannel->Initialize(InGroupInterface, GetOwningToolkit().GetLocalUser(), GroupId);
 
