#online One more net id tidy to avoid needless conversions

diff --git a/Engine/Plugins/Online/OnlineFramework/Source/Party/Public/Chat/SocialChatManager.cpp b/Engine/Plugins/Online/OnlineFramework/Source/Party/Public/Chat/SocialChatManager.cpp
--- a/Engine/Plugins/Online/OnlineFramework/Source/Party/Public/Chat/SocialChatManager.cpp
+++ b/Engine/Plugins/Online/OnlineFramework/Source/Party/Public/Chat/SocialChatManager.cpp
@@ -562,10 +562,8 @@ void USocialChatManager::OnGroupUpdated(const FUniqueNetId& GroupId)
 
 bool USocialChatManager::IsUniqueIdOfOwner(const FUniqueNetId& LocalUserId) const
 {
-	USocialUser& Owner = GetOwningToolkit().GetLocalUser();
-	FUniqueNetIdWrapper OwnerNetIdWrapper = FUniqueNetIdWrapper(Owner.GetUserId(ESocialSubsystem::Primary).GetUniqueNetId());
-	FUniqueNetIdWrapper LocalUserIdWrapper = FUniqueNetIdWrapper(LocalUserId);
-	return LocalUserIdWrapper == OwnerNetIdWrapper;
+	const USocialUser& Owner = GetOwningToolkit().GetLocalUser();
+	return LocalUserId == Owner.GetUserId(ESocialSubsystem::Primary);
 }
 
 void USocialChatManager::GetGroupChannels(TArray<USocialGroupChannel*>& JoinedChannels) const
