Provide a copy of the Party Member Data to PostReplication in case validation checks desire to roll back to previous data values.

#tests Compile DebugGame Editor Win64 / Shipping Client

diff --git a/Engine/Plugins/Online/OnlineFramework/Source/Party/Public/Party/PartyDataReplicator.h b/Engine/Plugins/Online/OnlineFramework/Source/Party/Public/Party/PartyDataReplicator.h
--- a/Engine/Plugins/Online/OnlineFramework/Source/Party/Public/Party/PartyDataReplicator.h
+++ b/Engine/Plugins/Online/OnlineFramework/Source/Party/Public/Party/PartyDataReplicator.h
@@ -65,7 +65,7 @@ PACKAGE_SCOPE:
 		{
 			if (FVariantDataConverter::VariantMapToUStruct(IncomingPartyData.GetKeyValAttrs(), RepDataType, RepDataPtr, 0, CPF_Transient | CPF_RepSkip))
 			{
-				static_cast<FOnlinePartyRepDataBase*>(RepDataPtr)->PostReplication();
+				static_cast<FOnlinePartyRepDataBase*>(RepDataPtr)->PostReplication(*RepDataCopy);
 
 				if (bCompareToPrevious)
 				{
diff --git a/Engine/Plugins/Online/OnlineFramework/Source/Party/Public/Party/PartyTypes.h b/Engine/Plugins/Online/OnlineFramework/Source/Party/Public/Party/PartyTypes.h
--- a/Engine/Plugins/Online/OnlineFramework/Source/Party/Public/Party/PartyTypes.h
+++ b/Engine/Plugins/Online/OnlineFramework/Source/Party/Public/Party/PartyTypes.h
@@ -409,7 +409,7 @@ protected:
 	/**
 	 * Called directly after an updated member state is received and copied into the local state
 	 */
-	virtual void PostReplication() {}
+	virtual void PostReplication(const FOnlinePartyRepDataBase& OldData) {}
 	virtual bool CanEditData() const { return false; }
 	virtual const USocialParty* GetOwnerParty() const { return nullptr; }
 	virtual const UPartyMember* GetOwningMember() const { return nullptr; }
