SocialParty->GetPartyMembers would cast first and then add to output array if cast was successful.

diff --git a/Engine/Plugins/Online/OnlineFramework/Source/Party/Public/Party/SocialParty.h b/Engine/Plugins/Online/OnlineFramework/Source/Party/Public/Party/SocialParty.h
--- a/Engine/Plugins/Online/OnlineFramework/Source/Party/Public/Party/SocialParty.h
+++ b/Engine/Plugins/Online/OnlineFramework/Source/Party/Public/Party/SocialParty.h
@@ -153,7 +153,10 @@ public:
 		PartyMembers.Reserve(PartyMembersById.Num());
 		for (const auto& IdMemberPair : PartyMembersById)
 		{
-			PartyMembers.Add(Cast<MemberT>(IdMemberPair.Value));
+			if (MemberT* PartyMember = Cast<MemberT>(IdMemberPair.Value))
+			{
+				PartyMembers.Add(PartyMember);
+			}
 		}
 		return PartyMembers;
 	}
