#online
Removed hotfix manager dependency from OnlineTracing

diff --git a/Engine/Plugins/Online/OnlineFramework/Source/Hotfix/Hotfix.Build.cs b/Engine/Plugins/Online/OnlineFramework/Source/Hotfix/Hotfix.Build.cs
--- a/Engine/Plugins/Online/OnlineFramework/Source/Hotfix/Hotfix.Build.cs
+++ b/Engine/Plugins/Online/OnlineFramework/Source/Hotfix/Hotfix.Build.cs
@@ -16,6 +16,7 @@ public class Hotfix : ModuleRules
                 "HTTP",
 				"OnlineSubsystem",
 				"OnlineSubsystemUtils",
+				"OnlineTracing"
 			}
 			);
     }
diff --git a/Engine/Plugins/Online/OnlineFramework/Source/Hotfix/Private/OnlineHotfixManager.cpp b/Engine/Plugins/Online/OnlineFramework/Source/Hotfix/Private/OnlineHotfixManager.cpp
--- a/Engine/Plugins/Online/OnlineFramework/Source/Hotfix/Private/OnlineHotfixManager.cpp
+++ b/Engine/Plugins/Online/OnlineFramework/Source/Hotfix/Private/OnlineHotfixManager.cpp
@@ -22,6 +22,7 @@
 #include "Engine/DataTable.h"
 #include "Curves/CurveFloat.h"
 #include "Engine/BlueprintGeneratedClass.h"
+#include "OnlineTracingModule.h"
 
 DEFINE_LOG_CATEGORY(LogHotfixManager);
 
@@ -856,6 +857,7 @@ bool UOnlineHotfixManager::HotfixIniFile(const FString& FileName, const FString&
 	bool bUpdateLogSuppression = false;
 	bool bUpdateConsoleVariables = false;
 	bool bUpdateHttpConfigs = false;
+	bool bUpdateOnlineTracing = false;
 	TSet<FString> OnlineSubSections;
 	// Find the set of object classes that were affected
 	while (StartIndex >= 0 && StartIndex < IniData.Len() && EndIndex >= StartIndex)
@@ -913,6 +915,7 @@ bool UOnlineHotfixManager::HotfixIniFile(const FString& FileName, const FString&
 					const TCHAR* ConsoleVariableSection = TEXT("[ConsoleVariables]");
 					const TCHAR* HttpSection = TEXT("[HTTP"); // note "]" omitted on purpose since we want a partial match
 					const TCHAR* OnlineSubSectionKey = TEXT("[OnlineSubsystem"); // note "]" omitted on purpose since we want a partial match
+					const TCHAR* OnlineTracingSection = TEXT("[OnlineTracing]");
 					if (!bUpdateLogSuppression && FCString::Strnicmp(*IniData + StartIndex, LogConfigSection, FCString::Strlen(LogConfigSection)) == 0)
 					{
 						bUpdateLogSuppression = true;
@@ -925,6 +928,10 @@ bool UOnlineHotfixManager::HotfixIniFile(const FString& FileName, const FString&
 					{
 						bUpdateHttpConfigs = true;
 					}
+					else if (!bUpdateOnlineTracing && FCString::Strnicmp(*IniData + StartIndex, OnlineTracingSection, FCString::Strlen(OnlineTracingSection)) == 0)
+					{
+						bUpdateOnlineTracing = true;
+					}
 					else if (FCString::Strnicmp(*IniData + StartIndex, OnlineSubSectionKey, FCString::Strlen(OnlineSubSectionKey)) == 0)
 					{
 						const FString SectionStr = IniData.Mid(StartIndex, EndIndex - StartIndex + 1);
@@ -1045,6 +1052,13 @@ bool UOnlineHotfixManager::HotfixIniFile(const FString& FileName, const FString&
 		FHttpModule::Get().UpdateConfigs();
 	}
 
+	// Reload configs for tracing system, this may init or tear it down if enable is toggled
+	if (bUpdateOnlineTracing && 
+		FOnlineTracingModule::IsAvailable())
+	{
+		FOnlineTracingModule::Get().UpdateConfig();
+	}
+
 	// Reload configs relevant to OSS config sections that were updated
 	IOnlineSubsystem* OnlineSub = IOnlineSubsystem::Get(OSSName.Len() ? FName(*OSSName, FNAME_Find) : NAME_None);
 	if (OnlineSub != nullptr)
