optional dependency on OnlineTracing since it is in NFL

diff --git a/Engine/Plugins/Online/OnlineFramework/Source/Hotfix/Hotfix.Build.cs b/Engine/Plugins/Online/OnlineFramework/Source/Hotfix/Hotfix.Build.cs
--- a/Engine/Plugins/Online/OnlineFramework/Source/Hotfix/Hotfix.Build.cs
+++ b/Engine/Plugins/Online/OnlineFramework/Source/Hotfix/Hotfix.Build.cs
@@ -1,5 +1,6 @@
 // Copyright Epic Games, Inc. All Rights Reserved.
 
+using System.IO;
 using UnrealBuildTool;
 
 public class Hotfix : ModuleRules
@@ -19,5 +20,12 @@ public class Hotfix : ModuleRules
 				"OnlineTracing"
 			}
 			);
-    }
+
+		bool bHasOnlineTracing = Directory.Exists(Path.Combine(EngineDirectory, "Plugins", "Online", "NotForLicensees", "OnlineTracing"));
+		if (bHasOnlineTracing)
+		{
+			PublicDefinitions.Add("WITH_ONLINETRACING=1");
+			PrivateDependencyModuleNames.Add("OnlineTracing");
+		}
+	}
 }
diff --git a/Engine/Plugins/Online/OnlineFramework/Source/Hotfix/Private/OnlineHotfixManager.cpp b/Engine/Plugins/Online/OnlineFramework/Source/Hotfix/Private/OnlineHotfixManager.cpp
--- a/Engine/Plugins/Online/OnlineFramework/Source/Hotfix/Private/OnlineHotfixManager.cpp
+++ b/Engine/Plugins/Online/OnlineFramework/Source/Hotfix/Private/OnlineHotfixManager.cpp
@@ -22,7 +22,10 @@
 #include "Engine/DataTable.h"
 #include "Curves/CurveFloat.h"
 #include "Engine/BlueprintGeneratedClass.h"
+
+#ifdef WITH_ONLINETRACING
 #include "OnlineTracingModule.h"
+#endif
 
 DEFINE_LOG_CATEGORY(LogHotfixManager);
 
@@ -1053,11 +1056,13 @@ bool UOnlineHotfixManager::HotfixIniFile(const FString& FileName, const FString&
 	}
 
 	// Reload configs for tracing system, this may init or tear it down if enable is toggled
+#ifdef WITH_ONLINETRACING
 	if (bUpdateOnlineTracing && 
 		FOnlineTracingModule::IsAvailable())
 	{
 		FOnlineTracingModule::Get().UpdateConfig();
 	}
+#endif
 
 	// Reload configs relevant to OSS config sections that were updated
 	IOnlineSubsystem* OnlineSub = IOnlineSubsystem::Get(OSSName.Len() ? FName(*OSSName, FNAME_Find) : NAME_None);
