missed a change

diff --git a/Engine/Plugins/Online/OnlineFramework/Source/Hotfix/Hotfix.Build.cs b/Engine/Plugins/Online/OnlineFramework/Source/Hotfix/Hotfix.Build.cs
--- a/Engine/Plugins/Online/OnlineFramework/Source/Hotfix/Hotfix.Build.cs
+++ b/Engine/Plugins/Online/OnlineFramework/Source/Hotfix/Hotfix.Build.cs
@@ -16,8 +16,7 @@ public class Hotfix : ModuleRules
 				"Engine",
                 "HTTP",
 				"OnlineSubsystem",
-				"OnlineSubsystemUtils",
-				"OnlineTracing"
+				"OnlineSubsystemUtils"
 			}
 			);
 
