temporarily revert dependency on OnlineTracing as it is missing for syncs without NFL access

diff --git a/Engine/Plugins/Online/OnlineFramework/Source/Hotfix/Hotfix.Build.cs b/Engine/Plugins/Online/OnlineFramework/Source/Hotfix/Hotfix.Build.cs
--- a/Engine/Plugins/Online/OnlineFramework/Source/Hotfix/Hotfix.Build.cs
+++ b/Engine/Plugins/Online/OnlineFramework/Source/Hotfix/Hotfix.Build.cs
@@ -19,12 +19,5 @@ public class Hotfix : ModuleRules
 				"OnlineSubsystemUtils"
 			}
 			);
-
-		bool bHasOnlineTracing = Directory.Exists(Path.Combine(EngineDirectory, "Plugins", "Online", "NotForLicensees", "OnlineTracing"));
-		if (bHasOnlineTracing)
-		{
-			PublicDefinitions.Add("WITH_ONLINETRACING=1");
-			PrivateDependencyModuleNames.Add("OnlineTracing");
-		}
 	}
 }
