default OnlineTracing plugin dependency to enabled so that dlls are included with the build and re-added it as a dependency to Hotfix module

diff --git a/Engine/Plugins/Online/OnlineFramework/Source/Hotfix/Hotfix.Build.cs b/Engine/Plugins/Online/OnlineFramework/Source/Hotfix/Hotfix.Build.cs
--- a/Engine/Plugins/Online/OnlineFramework/Source/Hotfix/Hotfix.Build.cs
+++ b/Engine/Plugins/Online/OnlineFramework/Source/Hotfix/Hotfix.Build.cs
@@ -19,5 +19,12 @@ public class Hotfix : ModuleRules
 				"OnlineSubsystemUtils"
 			}
 			);
+
+		bool bHasOnlineTracing = Directory.Exists(Path.Combine(EngineDirectory, "Plugins", "Online", "NotForLicensees", "OnlineTracing"));
+		if (bHasOnlineTracing)
+		{
+			PublicDefinitions.Add("WITH_ONLINETRACING=1");
+			PrivateDependencyModuleNames.Add("OnlineTracing");
+		}
 	}
 }
