Fix build warning due to missing plugin dependency

diff --git a/Engine/Plugins/Online/OnlineFramework/OnlineFramework.uplugin b/Engine/Plugins/Online/OnlineFramework/OnlineFramework.uplugin
--- a/Engine/Plugins/Online/OnlineFramework/OnlineFramework.uplugin
+++ b/Engine/Plugins/Online/OnlineFramework/OnlineFramework.uplugin
@@ -54,6 +54,11 @@
 		{
 			"Name": "OnlineSubsystemUtils",
 			"Enabled": true
+		},
+		{
+			"Name": "OnlineTracing",
+			"Enabled": true,
+			"Optional": true
 		}
 	]
 }
\ No newline at end of file
