Add a cheat command in SocialToolkit to randomly send out presence changed event and simulate some of the presence info to move the users around in the list (online status, party size)

SocialUI.StartRandomizeUserPresence uint8 (NumRandomUser) float(TickerTimer)
SocialUI.StopRandomizeUserPresence bool (bClearGeneratedPresence)

diff --git a/Engine/Plugins/Online/OnlineFramework/Source/Party/Private/SocialToolkit.cpp b/Engine/Plugins/Online/OnlineFramework/Source/Party/Private/SocialToolkit.cpp
--- a/Engine/Plugins/Online/OnlineFramework/Source/Party/Private/SocialToolkit.cpp
+++ b/Engine/Plugins/Online/OnlineFramework/Source/Party/Private/SocialToolkit.cpp
@@ -12,6 +12,49 @@
 #include "OnlineSubsystemUtils.h"
 #include "Engine/LocalPlayer.h"
 
+#if WITH_EDITOR
+DECLARE_MULTICAST_DELEGATE_TwoParams(FOnStartRandomizeUserPresence, uint8 /*NumRandomUser*/, float /*TickerTimer*/);
+static FOnStartRandomizeUserPresence Debug_OnStartRandomizeUserPresenceEvent;
+static FAutoConsoleCommandWithWorldAndArgs CMD_OnStartRandomUserPresence
+(
+	TEXT("SocialUI.StartRandomizeUserPresence"),
+	TEXT("Randomize users' presence and fire off presence changed events to trigger friend list update/refresh etc. @param NumRandomUser uint8, @param TickerTimer float"),
+	FConsoleCommandWithWorldAndArgsDelegate::CreateLambda([](const TArray<FString>&Args, UWorld* World)
+{
+	uint8 NumRandomUser = 2;
+	float TickerTimer = 5.f;
+
+	if (Args.Num() > 0)
+	{
+		NumRandomUser = FCString::Atoi(*Args[0]);
+	}
+	if (Args.Num() > 1)
+	{
+		TickerTimer = FCString::Atof(*Args[1]);
+	}
+	Debug_OnStartRandomizeUserPresenceEvent.Broadcast(NumRandomUser, TickerTimer);
+})
+);
+
+DECLARE_MULTICAST_DELEGATE_OneParam(FOnStopRandomizeUserPresence, bool /*bClearGeneratedPresence*/)
+static FOnStopRandomizeUserPresence Debug_OnStopRandomizeUserPresenceEvent;
+static FAutoConsoleCommandWithWorldAndArgs CMD_OnStopRandomUserPresence
+(
+	TEXT("SocialUI.StopRandomizeUserPresence"),
+	TEXT("Stop randomizing users' presnece, with optional param to clear off already generated presence (default to false). @param bClearGeneratedPresence uint8"),
+	FConsoleCommandWithWorldAndArgsDelegate::CreateLambda([](const TArray<FString>&Args, UWorld* World)
+{
+	bool bClearGeneratedPresence = false;
+	if (Args.Num() > 0)
+	{
+		bClearGeneratedPresence = Args[0].ToBool();
+	}
+	Debug_OnStopRandomizeUserPresenceEvent.Broadcast(bClearGeneratedPresence);
+})
+);
+
+#endif
+
 bool NameToSocialSubsystem(FName SubsystemName, ESocialSubsystem& OutSocialSubsystem)
 {
 	for (uint8 SocialSubsystemIdx = 0; SocialSubsystemIdx < (uint8)ESocialSubsystem::MAX; ++SocialSubsystemIdx)
@@ -125,6 +168,11 @@ USocialToolkit::USocialToolkit()
 
 void USocialToolkit::InitializeToolkit(ULocalPlayer& InOwningLocalPlayer)
 {
+#if WITH_EDITOR
+	Debug_OnStartRandomizeUserPresenceEvent.AddUObject(this, &USocialToolkit::Debug_OnStartRandomizeUserPresence);
+	Debug_OnStopRandomizeUserPresenceEvent.AddUObject(this, &USocialToolkit::Debug_OnStopRandomizeUserPresence);
+#endif
+
 	LocalPlayerOwner = &InOwningLocalPlayer;
 
 	SocialChatManager = USocialChatManager::CreateChatManager(*this);
@@ -1030,3 +1078,60 @@ void USocialToolkit::HandleExistingPartyInvites(ESocialSubsystem SubsystemType)
 		}
 	}
 }
+
+#if WITH_EDITOR
+void USocialToolkit::Debug_OnStartRandomizeUserPresence(uint8 NumRandomUser, float TickerTimer)
+{
+	if (Debug_PresenceTickerHandle.IsValid())
+	{
+		FTicker::GetCoreTicker().RemoveTicker(Debug_PresenceTickerHandle);
+		Debug_PresenceTickerHandle.Reset();
+	}
+
+	if (ensure(TickerTimer > 0.f))
+	{
+		bDebug_IsRandomlyChangingUserPresence = true;
+		Debug_PresenceTickerHandle = FTicker::GetCoreTicker().AddTicker(FTickerDelegate::CreateUObject(this, &USocialToolkit::Debug_HandleRandomizeUserPresenceTick, NumRandomUser), TickerTimer);
+	}
+}
+
+void USocialToolkit::Debug_OnStopRandomizeUserPresence(bool bClearGeneratedPresence)
+{
+	bDebug_IsRandomlyChangingUserPresence = false;
+	if (Debug_PresenceTickerHandle.IsValid())
+	{
+		FTicker::GetCoreTicker().RemoveTicker(Debug_PresenceTickerHandle);
+		Debug_PresenceTickerHandle.Reset();
+	}
+
+	if (bClearGeneratedPresence)
+	{
+		// Refresh all existing presence data to revert them back to normal
+		for (USocialUser* User : AllUsers)
+		{
+			User->bDebug_IsPresenceArtificial = false;
+			User->NotifyPresenceChanged(ESocialSubsystem::Primary);
+		}
+	}
+}
+
+bool USocialToolkit::Debug_HandleRandomizeUserPresenceTick(float DeltaTime, uint8 NumRandomUser)
+{
+	Debug_ChangeRandomUserPresence(NumRandomUser);
+	return true;
+}
+
+void USocialToolkit::Debug_ChangeRandomUserPresence(uint8 NumRandomUser)
+{
+	const TArray<USocialUser*> SocialUsers = GetAllUsers();
+	for (int32 i = 0; i < NumRandomUser; i++)
+	{
+		int32 UserIndex = FMath::RandRange(0, SocialUsers.Num() - 1);
+		if (USocialUser* SocialUser = SocialUsers[UserIndex])
+		{
+			SocialUser->Debug_RandomizePresence();
+			SocialUser->NotifyPresenceChanged(ESocialSubsystem::Primary);
+		}
+	}
+}
+#endif
\ No newline at end of file
diff --git a/Engine/Plugins/Online/OnlineFramework/Source/Party/Private/User/SocialUser.cpp b/Engine/Plugins/Online/OnlineFramework/Source/Party/Private/User/SocialUser.cpp
--- a/Engine/Plugins/Online/OnlineFramework/Source/Party/Private/User/SocialUser.cpp
+++ b/Engine/Plugins/Online/OnlineFramework/Source/Party/Private/User/SocialUser.cpp
@@ -305,6 +305,13 @@ EOnlinePresenceState::Type USocialUser::GetOnlineStatus() const
 		return EOnlinePresenceState::Offline;
 	}
 
+#if WITH_EDITOR
+	if (bDebug_IsPresenceArtificial && !IsBlocked())
+	{
+		return Debug_RandomPresence;
+	}
+#endif
+
 	EOnlinePresenceState::Type OnlineStatus = EOnlinePresenceState::Offline;
 
 	// Get the most "present" status available on any of the associated platforms
@@ -1200,6 +1207,14 @@ void USocialUser::EstablishOssInfo(const TSharedRef<FOnlineRecentPlayer>& InRece
 	}
 }
 
+#if WITH_EDITOR
+void USocialUser::Debug_RandomizePresence()
+{
+	bDebug_IsPresenceArtificial = true;
+	Debug_RandomPresence = static_cast<EOnlinePresenceState::Type>(FMath::RandRange((int32)EOnlinePresenceState::Online, (int32)EOnlinePresenceState::Away));
+}
+#endif
+
 void USocialUser::OnPresenceChangedInternal(ESocialSubsystem SubsystemType)
 {
 	OnUserPresenceChanged().Broadcast(SubsystemType);
diff --git a/Engine/Plugins/Online/OnlineFramework/Source/Party/Private/User/SocialUserList.cpp b/Engine/Plugins/Online/OnlineFramework/Source/Party/Private/User/SocialUserList.cpp
--- a/Engine/Plugins/Online/OnlineFramework/Source/Party/Private/User/SocialUserList.cpp
+++ b/Engine/Plugins/Online/OnlineFramework/Source/Party/Private/User/SocialUserList.cpp
@@ -399,6 +399,14 @@ bool FSocialUserList::EvaluateUserPresence(const USocialUser& User, ESocialSubsy
 			{
 				bInSameParty = CurrentParty->ContainsUser(User);
 			}
+
+#if WITH_EDITOR
+			if (OwnerToolkit->Debug_IsRandomlyChangingPresence())
+			{
+				bIsOnline = User.GetOnlineStatus() != EOnlinePresenceState::Offline;
+				bIsPlayingThisGame = bIsOnline;
+			}
+#endif
 		}
 
 		return EvaluatePresenceFlag(bIsOnline, ESocialUserStateFlags::Online)
diff --git a/Engine/Plugins/Online/OnlineFramework/Source/Party/Public/SocialToolkit.h b/Engine/Plugins/Online/OnlineFramework/Source/Party/Public/SocialToolkit.h
--- a/Engine/Plugins/Online/OnlineFramework/Source/Party/Public/SocialToolkit.h
+++ b/Engine/Plugins/Online/OnlineFramework/Source/Party/Public/SocialToolkit.h
@@ -110,6 +110,10 @@ public:
 	DECLARE_EVENT(USocialToolkit, FBasicToolkitEvent);
 	FBasicToolkitEvent& OnToolkitReset() const { return OnToolkitResetEvent; }
 
+#if WITH_EDITOR
+	bool Debug_IsRandomlyChangingPresence() { return bDebug_IsRandomlyChangingUserPresence; }
+#endif
+
 PARTY_SCOPE:
 	void NotifySubsystemIdEstablished(USocialUser& SocialUser, ESocialSubsystem SubsystemType, const FUniqueNetIdRepl& SubsystemId);
 	TSubclassOf<USocialChatManager> GetChatManagerClass() { return ChatManagerClass; }
@@ -221,6 +225,15 @@ private:	// Handlers
 
 	void HandleExistingPartyInvites(ESocialSubsystem SubsystemType);
 
+#if WITH_EDITOR
+	void Debug_OnStartRandomizeUserPresence(uint8 NumRandomUser, float TickerTimer);
+	void Debug_OnStopRandomizeUserPresence(bool bClearGeneratedPresence);
+	bool Debug_HandleRandomizeUserPresenceTick(float DeltaTime, uint8 NumRandomUser);
+	void Debug_ChangeRandomUserPresence(uint8 NumRandomUser);
+	bool bDebug_IsRandomlyChangingUserPresence = false;
+	FDelegateHandle Debug_PresenceTickerHandle;
+#endif
+
 private:
 	static USocialToolkit* GetToolkitForPlayerInternal(ULocalPlayer* LocalPlayer);
 	static TMap<TWeakObjectPtr<ULocalPlayer>, TWeakObjectPtr<USocialToolkit>> AllToolkitsByOwningPlayer;
diff --git a/Engine/Plugins/Online/OnlineFramework/Source/Party/Public/User/SocialUser.h b/Engine/Plugins/Online/OnlineFramework/Source/Party/Public/User/SocialUser.h
--- a/Engine/Plugins/Online/OnlineFramework/Source/Party/Public/User/SocialUser.h
+++ b/Engine/Plugins/Online/OnlineFramework/Source/Party/Public/User/SocialUser.h
@@ -149,6 +149,12 @@ PARTY_SCOPE:
 	void EstablishOssInfo(const TSharedRef<FOnlineBlockedPlayer>& BlockedPlayerInfo, ESocialSubsystem SubsystemType);
 	void EstablishOssInfo(const TSharedRef<FOnlineRecentPlayer>& RecentPlayerInfo, ESocialSubsystem SubsystemType);
 
+#if WITH_EDITOR
+	void Debug_RandomizePresence();
+	bool bDebug_IsPresenceArtificial = false;
+	EOnlinePresenceState::Type Debug_RandomPresence;
+#endif
+
 protected:
 	virtual void OnPresenceChangedInternal(ESocialSubsystem SubsystemType);
 	virtual void OnPartyInviteAcceptedInternal(const FOnlinePartyTypeId& PartyTypeId) const;
