Add option to disable sorting during SocialUserList updates.

Often the order of the items in the SocialUserList does not matter. We simply do a linear walk through the list or we sort social users into other containers for temporary purposes.
A key point of this optimization is to preserve existing functionality for licensees.

#tests Compile DebugGame Editor Win64 / Development Editor Win64 / Shipping Client / Test Client / Ran it in PIE / Ran in Test

diff --git a/Engine/Plugins/Online/OnlineFramework/Source/Party/Private/User/SocialUserList.cpp b/Engine/Plugins/Online/OnlineFramework/Source/Party/Private/User/SocialUserList.cpp
--- a/Engine/Plugins/Online/OnlineFramework/Source/Party/Private/User/SocialUserList.cpp
+++ b/Engine/Plugins/Online/OnlineFramework/Source/Party/Private/User/SocialUserList.cpp
@@ -107,6 +107,11 @@ void FSocialUserList::SetAllowAutoUpdate(bool bIsEnabled)
 	}
 }
 
+void FSocialUserList::SetAllowSortDuringUpdate(bool bIsEnabled)
+{
+	ListConfig.bSortDuringUpdate = bIsEnabled;
+}
+
 bool FSocialUserList::HasPresenceFilters() const
 {
 	return ListConfig.RequiredPresenceFlags != ESocialUserStateFlags::None || ListConfig.ForbiddenPresenceFlags != ESocialUserStateFlags::None;
@@ -518,10 +523,10 @@ void FSocialUserList::UpdateListInternal()
 	UsersWithDirtyPresence.Reset();
 
 	// Update the users in the list
-	bool bListChanged = false;
+	bool bListUpdated = false;
 	if (PendingRemovals.Num() > 0)
 	{
-		bListChanged = true;
+		bListUpdated = true;
 
 		Users.RemoveAllSwap(
 			[this] (USocialUser* User)
@@ -541,7 +546,7 @@ void FSocialUserList::UpdateListInternal()
 
 	if (PendingAdds.Num() > 0)
 	{
-		bListChanged = true;
+		bListUpdated = true;
 		Users.Append(PendingAdds);
 
 		for (USocialUser* User : PendingAdds)
@@ -551,28 +556,46 @@ void FSocialUserList::UpdateListInternal()
 		PendingAdds.Reset();
 	}
 
-	if (bListChanged || bNeedsSort)
+	if (bListUpdated || bNeedsSort)
 	{
-		bNeedsSort = false;
+		if (ListConfig.bSortDuringUpdate)
+		{
+			bNeedsSort = false;
+			bListUpdated = true;
+
+			const int32 NumUsers = Users.Num();
+			if (NumUsers > 1)
+			{
+				SCOPED_NAMED_EVENT(STAT_SocialUserList_Sort, FColor::Orange);
 
-		const int32 NumUsers = Users.Num();
-		TArray<FUserSortData> SortedData;
-		SortedData.Reserve(NumUsers);
+				UE_LOG(LogParty, Verbose, TEXT("%s sorting list of [%d] users"), ANSI_TO_TCHAR(__FUNCTION__), NumUsers);
 
-		Algo::Transform(Users, SortedData, [](USocialUser* const User) -> FUserSortData
-		{
-			return FUserSortData(User, User->GetOnlineStatus(), User->IsPlayingThisGame(), User->GetDisplayName(), User->GetCustomSortValuePrimary(), User->GetCustomSortValueSecondary());
-		});
+				TArray<FUserSortData> SortedData;
+				SortedData.Reserve(NumUsers);
 
-		Algo::Sort(SortedData);
+				Algo::Transform(Users, SortedData, [](USocialUser* const User) -> FUserSortData
+				{
+					return FUserSortData(User, User->GetOnlineStatus(), User->IsPlayingThisGame(), User->GetDisplayName(), User->GetCustomSortValuePrimary(), User->GetCustomSortValueSecondary());
+				});
 
-		// replace contents of Users from SortedData array
-		for (int Index = 0; Index < NumUsers; Index++)
+				Algo::Sort(SortedData);
+
+				// replace contents of Users from SortedData array
+				for (int Index = 0; Index < NumUsers; Index++)
+				{
+					Users[Index] = SortedData[Index].User;
+				}
+			}
+		}
+		else
 		{
-			Users[Index] = SortedData[Index].User;
+			bNeedsSort = true;
+		}
+
+		if (bListUpdated)
+		{
+			OnUpdateComplete().Broadcast();
 		}
-		
-		OnUpdateComplete().Broadcast();
 	}
 }
 
diff --git a/Engine/Plugins/Online/OnlineFramework/Source/Party/Private/User/SocialUserList.h b/Engine/Plugins/Online/OnlineFramework/Source/Party/Private/User/SocialUserList.h
--- a/Engine/Plugins/Online/OnlineFramework/Source/Party/Private/User/SocialUserList.h
+++ b/Engine/Plugins/Online/OnlineFramework/Source/Party/Private/User/SocialUserList.h
@@ -22,6 +22,7 @@ public:
 
 	void UpdateNow();
 	void SetAllowAutoUpdate(bool bIsEnabled);
+	void SetAllowSortDuringUpdate(bool bIsEnabled);
 	const TArray<USocialUser*>& GetUsers() const { return Users; }
 
 	bool HasPresenceFilters() const;
@@ -82,8 +83,6 @@ private:
 
 	FSocialUserListConfig ListConfig;
 
-	// give external access to disable list update for perf
-	bool bAllowAutoUpdate = true;
 	bool bNeedsSort = false;
 	float AutoUpdatePeriod = 5.f;
 	FDelegateHandle UpdateTickerHandle;
diff --git a/Engine/Plugins/Online/OnlineFramework/Source/Party/Public/User/ISocialUserList.h b/Engine/Plugins/Online/OnlineFramework/Source/Party/Public/User/ISocialUserList.h
--- a/Engine/Plugins/Online/OnlineFramework/Source/Party/Public/User/ISocialUserList.h
+++ b/Engine/Plugins/Online/OnlineFramework/Source/Party/Public/User/ISocialUserList.h
@@ -44,6 +44,8 @@ public:
 
 	// Whether or not the list should be polled regularly for updates (as opposed to manually having UpdateNow triggered)
 	bool bAutoUpdate = false;
+
+	bool bSortDuringUpdate = true;
 };
 
 class ISocialUserList
@@ -70,4 +72,6 @@ public:
 
 	/** Give external overwrite to disable list auto update for perf */
 	virtual void SetAllowAutoUpdate(bool bIsEnabled) = 0;
+
+	virtual void SetAllowSortDuringUpdate(bool bIsEnabled) = 0;
 };
\ No newline at end of file
