Social Toolkit: GetToolkitForPlayer now accepts a const ULocalPlayer pointer.

diff --git a/Engine/Plugins/Online/OnlineFramework/Source/Party/Private/SocialToolkit.cpp b/Engine/Plugins/Online/OnlineFramework/Source/Party/Private/SocialToolkit.cpp
--- a/Engine/Plugins/Online/OnlineFramework/Source/Party/Private/SocialToolkit.cpp
+++ b/Engine/Plugins/Online/OnlineFramework/Source/Party/Private/SocialToolkit.cpp
@@ -153,7 +153,7 @@ public:
 
 //@todo DanH Social: Need a non-backdoor way to get toolkits from the manager (an issue when we don't know where the manager is) - new game subsystems should be a nice solve
 TMap<TWeakObjectPtr<ULocalPlayer>, TWeakObjectPtr<USocialToolkit>> USocialToolkit::AllToolkitsByOwningPlayer;
-USocialToolkit* USocialToolkit::GetToolkitForPlayerInternal(ULocalPlayer* LocalPlayer)
+USocialToolkit* USocialToolkit::GetToolkitForPlayerInternal(const ULocalPlayer* LocalPlayer)
 {
 	TWeakObjectPtr<USocialToolkit>* FoundToolkit = AllToolkitsByOwningPlayer.Find(LocalPlayer);
 	return FoundToolkit ? FoundToolkit->Get() : nullptr;
diff --git a/Engine/Plugins/Online/OnlineFramework/Source/Party/Public/SocialToolkit.h b/Engine/Plugins/Online/OnlineFramework/Source/Party/Public/SocialToolkit.h
--- a/Engine/Plugins/Online/OnlineFramework/Source/Party/Public/SocialToolkit.h
+++ b/Engine/Plugins/Online/OnlineFramework/Source/Party/Public/SocialToolkit.h
@@ -38,7 +38,7 @@ class PARTY_API USocialToolkit : public UObject
 
 public:
 	template <typename ToolkitT = USocialToolkit>
-	static ToolkitT* GetToolkitForPlayer(ULocalPlayer* LocalPlayer)
+	static ToolkitT* GetToolkitForPlayer(const ULocalPlayer* LocalPlayer)
 	{
 		static_assert(TIsDerivedFrom<ToolkitT, USocialToolkit>::IsDerived, "GetToolkitForPlayer only supports getting USocialToolkit type objects");
 		return Cast<ToolkitT>(GetToolkitForPlayerInternal(LocalPlayer));
@@ -235,7 +235,7 @@ private:	// Handlers
 #endif
 
 private:
-	static USocialToolkit* GetToolkitForPlayerInternal(ULocalPlayer* LocalPlayer);
+	static USocialToolkit* GetToolkitForPlayerInternal(const ULocalPlayer* LocalPlayer);
 	static TMap<TWeakObjectPtr<ULocalPlayer>, TWeakObjectPtr<USocialToolkit>> AllToolkitsByOwningPlayer;
 
 	UPROPERTY()
