Fixed Linux/Mac/etc deprecation warnings on 11072805.

diff --git a/Engine/Plugins/Online/OnlineFramework/Source/Party/Private/SocialToolkit.cpp b/Engine/Plugins/Online/OnlineFramework/Source/Party/Private/SocialToolkit.cpp
--- a/Engine/Plugins/Online/OnlineFramework/Source/Party/Private/SocialToolkit.cpp
+++ b/Engine/Plugins/Online/OnlineFramework/Source/Party/Private/SocialToolkit.cpp
@@ -152,7 +152,7 @@ public:
 //////////////////////////////////////////////////////////////////////////
 
 //@todo DanH Social: Need a non-backdoor way to get toolkits from the manager (an issue when we don't know where the manager is) - new game subsystems should be a nice solve
-TMap<TWeakObjectPtr<ULocalPlayer>, TWeakObjectPtr<USocialToolkit>> USocialToolkit::AllToolkitsByOwningPlayer;
+TMap<TWeakObjectPtr<const ULocalPlayer>, TWeakObjectPtr<USocialToolkit>> USocialToolkit::AllToolkitsByOwningPlayer;
 USocialToolkit* USocialToolkit::GetToolkitForPlayerInternal(const ULocalPlayer* LocalPlayer)
 {
 	TWeakObjectPtr<USocialToolkit>* FoundToolkit = AllToolkitsByOwningPlayer.Find(LocalPlayer);
diff --git a/Engine/Plugins/Online/OnlineFramework/Source/Party/Public/SocialToolkit.h b/Engine/Plugins/Online/OnlineFramework/Source/Party/Public/SocialToolkit.h
--- a/Engine/Plugins/Online/OnlineFramework/Source/Party/Public/SocialToolkit.h
+++ b/Engine/Plugins/Online/OnlineFramework/Source/Party/Public/SocialToolkit.h
@@ -236,7 +236,7 @@ private:	// Handlers
 
 private:
 	static USocialToolkit* GetToolkitForPlayerInternal(const ULocalPlayer* LocalPlayer);
-	static TMap<TWeakObjectPtr<ULocalPlayer>, TWeakObjectPtr<USocialToolkit>> AllToolkitsByOwningPlayer;
+	static TMap<TWeakObjectPtr<const ULocalPlayer>, TWeakObjectPtr<USocialToolkit>> AllToolkitsByOwningPlayer;
 
 	UPROPERTY()
 	USocialUser* LocalUser;
