Added USocialToolkit::OnSocialUserInvalidated() event, to help callers process and handle SocialUsers that turn out to be uninitializable.

#tests Compile DebugGame Editor Win64 / Shipping Client

diff --git a/Engine/Plugins/Online/OnlineFramework/Source/Party/Private/SocialToolkit.cpp b/Engine/Plugins/Online/OnlineFramework/Source/Party/Private/SocialToolkit.cpp
--- a/Engine/Plugins/Online/OnlineFramework/Source/Party/Private/SocialToolkit.cpp
+++ b/Engine/Plugins/Online/OnlineFramework/Source/Party/Private/SocialToolkit.cpp
@@ -1053,9 +1053,10 @@ void USocialToolkit::HandleGameDestroyed(const FName SessionName, bool bWasSucce
 	// Update the recent player list whenever a game session ends
 }
 
-void USocialToolkit::HandleUserInvalidated(USocialUser* InvalidUser)
+void USocialToolkit::HandleUserInvalidated(USocialUser& InvalidUser)
 {
-	AllUsers.Remove(InvalidUser);
+	AllUsers.Remove(&InvalidUser);
+	OnSocialUserInvalidated().Broadcast(InvalidUser);
 }
 
 void USocialToolkit::HandleExistingPartyInvites(ESocialSubsystem SubsystemType)
diff --git a/Engine/Plugins/Online/OnlineFramework/Source/Party/Private/User/SocialUser.cpp b/Engine/Plugins/Online/OnlineFramework/Source/Party/Private/User/SocialUser.cpp
--- a/Engine/Plugins/Online/OnlineFramework/Source/Party/Private/User/SocialUser.cpp
+++ b/Engine/Plugins/Online/OnlineFramework/Source/Party/Private/User/SocialUser.cpp
@@ -358,7 +358,7 @@ void USocialUser::TryBroadcastInitializationComplete()
 				InitEventsByUser.Remove(this);
 
 				// Remove Toolkit's reference to the SocialUser, GC will clean it
-				GetOwningToolkit().HandleUserInvalidated(this);
+				GetOwningToolkit().HandleUserInvalidated(*this);
 			}
 		}
 	}
diff --git a/Engine/Plugins/Online/OnlineFramework/Source/Party/Public/SocialToolkit.h b/Engine/Plugins/Online/OnlineFramework/Source/Party/Public/SocialToolkit.h
--- a/Engine/Plugins/Online/OnlineFramework/Source/Party/Public/SocialToolkit.h
+++ b/Engine/Plugins/Online/OnlineFramework/Source/Party/Public/SocialToolkit.h
@@ -107,6 +107,9 @@ public:
 	DECLARE_EVENT_OneParam(USocialToolkit, FOnKnownUserInitialized, USocialUser&);
 	FOnKnownUserInitialized& OnKnownUserInitialized() { return OnKnownUserInitializedEvent; }
 
+	DECLARE_EVENT_OneParam(USocialToolkit, FOnSocialUserInvalidated, const USocialUser&);
+	FOnSocialUserInvalidated& OnSocialUserInvalidated() { return OnSocialUserInvalidatedEvent; }
+
 	DECLARE_EVENT(USocialToolkit, FBasicToolkitEvent);
 	FBasicToolkitEvent& OnToolkitReset() const { return OnToolkitResetEvent; }
 
@@ -122,7 +125,7 @@ PARTY_SCOPE:
 
 	bool AcceptFriendInvite(const USocialUser& SocialUser, ESocialSubsystem SubsystemType) const;
 
-	void HandleUserInvalidated(USocialUser* InvalidUser);
+	void HandleUserInvalidated(USocialUser& InvalidUser);
 
 #if PLATFORM_PS4
 	void NotifyPSNFriendsListRebuilt();
@@ -263,5 +266,6 @@ private:
 	mutable FOnRelationshipEstablished OnRecentPlayerAddedEvent;
 	
 	mutable FOnKnownUserInitialized OnKnownUserInitializedEvent;
+	mutable FOnSocialUserInvalidated OnSocialUserInvalidatedEvent;
 	mutable FBasicToolkitEvent OnToolkitResetEvent;
 };
