- Added function to Social Party to check if a local player is the party leader.
- Added function to Social Party to check if the party leader is a local player.

#rnx

diff --git a/Engine/Plugins/Online/OnlineFramework/Source/Party/Private/Party/SocialParty.cpp b/Engine/Plugins/Online/OnlineFramework/Source/Party/Private/Party/SocialParty.cpp
--- a/Engine/Plugins/Online/OnlineFramework/Source/Party/Private/Party/SocialParty.cpp
+++ b/Engine/Plugins/Online/OnlineFramework/Source/Party/Private/Party/SocialParty.cpp
@@ -1287,6 +1287,16 @@ bool USocialParty::IsLocalPlayerPartyLeader() const
 	return OwningLocalUserId == CurrentLeaderId;
 }
 
+bool USocialParty::IsPartyLeader(const ULocalPlayer& LocalPlayer) const
+{
+	return LocalPlayer.GetPreferredUniqueNetId() == CurrentLeaderId;
+}
+
+bool USocialParty::IsPartyLeaderLocal() const
+{
+	return GetSocialManager().IsLocalUser(CurrentLeaderId, ESocialSubsystem::Primary);
+}
+
 bool USocialParty::IsNetDriverFromReservationBeacon(const UNetDriver* const InNetDriver) const
 {
 	const FName NetDriverName = InNetDriver->NetDriverName;
diff --git a/Engine/Plugins/Online/OnlineFramework/Source/Party/Public/Party/SocialParty.h b/Engine/Plugins/Online/OnlineFramework/Source/Party/Public/Party/SocialParty.h
--- a/Engine/Plugins/Online/OnlineFramework/Source/Party/Public/Party/SocialParty.h
+++ b/Engine/Plugins/Online/OnlineFramework/Source/Party/Public/Party/SocialParty.h
@@ -122,6 +122,8 @@ public:
 	const FUniqueNetIdRepl& GetOwningLocalUserId() const { return OwningLocalUserId; }
 	const FUniqueNetIdRepl& GetPartyLeaderId() const { return CurrentLeaderId; }
 	bool IsLocalPlayerPartyLeader() const;
+	bool IsPartyLeader(const ULocalPlayer& LocalPlayer) const;
+	bool IsPartyLeaderLocal() const;
 
 	FChatRoomId GetChatRoomId() const;
 	bool IsPersistentParty() const;
