Socialtoolkit starts to cache all user lists made, added list config comparision to match at least lists with same configs and without custome filters. CreateUserList will now return a cached user list if the list config requested matches.
  -- Also added a counter of how many listeners are actually requesting a list to be audo updated, so one place turning auto update off for perf doesn't also turn it off for other places that are still listening for auto updates.

diff --git a/Engine/Plugins/Online/OnlineFramework/Source/Party/Private/SocialToolkit.cpp b/Engine/Plugins/Online/OnlineFramework/Source/Party/Private/SocialToolkit.cpp
--- a/Engine/Plugins/Online/OnlineFramework/Source/Party/Private/SocialToolkit.cpp
+++ b/Engine/Plugins/Online/OnlineFramework/Source/Party/Private/SocialToolkit.cpp
@@ -209,9 +209,27 @@ IOnlineSubsystem* USocialToolkit::GetSocialOss(ESocialSubsystem SubsystemType) c
 	return Online::GetSubsystem(GetWorld(), USocialManager::GetSocialOssName(SubsystemType));
 }
 
-TSharedRef<ISocialUserList> USocialToolkit::CreateUserList(const FSocialUserListConfig& ListConfig)
+TSharedRef<ISocialUserList> USocialToolkit::CreateUserList(const FSocialUserListConfig& ListConfig) const
 {
-	return FSocialUserList::CreateUserList(*this, ListConfig);
+	CachedSocialUserLists.RemoveAll([](const TWeakPtr<FSocialUserList>& UserList)
+	{
+		return !UserList.IsValid(); 
+	});
+
+	TWeakPtr<FSocialUserList>* FoundUserList = CachedSocialUserLists.FindByPredicate([ListConfig](const TWeakPtr<FSocialUserList>& UserList)
+	{
+		return UserList.IsValid() ? UserList.Pin()->GetListConfig() == ListConfig : false;
+	});
+
+	if (FoundUserList && (*FoundUserList).IsValid())
+	{
+		UE_LOG(LogParty, Verbose, TEXT("%s Found Userlist %s while creating Userlist %s with the same list config."), ANSI_TO_TCHAR(__FUNCTION__), *(*FoundUserList).Pin()->GetListConfig().Name, *ListConfig.Name);
+		return (*FoundUserList).Pin().ToSharedRef();
+	}
+	
+	TSharedRef<FSocialUserList> NewUserList = FSocialUserList::CreateUserList(*this, ListConfig);
+	CachedSocialUserLists.Add(NewUserList);
+	return NewUserList;
 }
 
 USocialUser& USocialToolkit::GetLocalUser() const
diff --git a/Engine/Plugins/Online/OnlineFramework/Source/Party/Private/User/SocialUserList.cpp b/Engine/Plugins/Online/OnlineFramework/Source/Party/Private/User/SocialUserList.cpp
--- a/Engine/Plugins/Online/OnlineFramework/Source/Party/Private/User/SocialUserList.cpp
+++ b/Engine/Plugins/Online/OnlineFramework/Source/Party/Private/User/SocialUserList.cpp
@@ -6,20 +6,19 @@
 #include "Party/PartyMember.h"
 
 #include "Containers/Ticker.h"
-#include "Interfaces/OnlinePresenceInterface.h"
 #include "SocialSettings.h"
 #include "Algo/Transform.h"
 #include "SocialManager.h"
 #include "Party/SocialParty.h"
 
-TSharedRef<FSocialUserList> FSocialUserList::CreateUserList(USocialToolkit& InOwnerToolkit, const FSocialUserListConfig& InConfig)
+TSharedRef<FSocialUserList> FSocialUserList::CreateUserList(const USocialToolkit& InOwnerToolkit, const FSocialUserListConfig& InConfig)
 {
 	TSharedRef<FSocialUserList> NewList = MakeShareable(new FSocialUserList(InOwnerToolkit, InConfig));
 	NewList->InitializeList();
 	return NewList;
 }
 
-FSocialUserList::FSocialUserList(USocialToolkit& InOwnerToolkit, const FSocialUserListConfig& InConfig)
+FSocialUserList::FSocialUserList(const USocialToolkit& InOwnerToolkit, const FSocialUserListConfig& InConfig)
 	: OwnerToolkit(&InOwnerToolkit)
 	, ListConfig(InConfig)
 {
@@ -96,7 +95,10 @@ void FSocialUserList::UpdateNow()
 
 void FSocialUserList::SetAllowAutoUpdate(bool bIsEnabled)
 {
-	if (!bIsEnabled && UpdateTickerHandle.IsValid())
+	bIsEnabled ? AutoUpdateRequests++ : AutoUpdateRequests--;
+	AutoUpdateRequests = FMath::Max(AutoUpdateRequests, 0);
+
+	if (AutoUpdateRequests == 0 && UpdateTickerHandle.IsValid())
 	{
 		FTicker::GetCoreTicker().RemoveTicker(UpdateTickerHandle);
 		UpdateTickerHandle.Reset();
@@ -225,7 +227,6 @@ void FSocialUserList::HandleUserPresenceChanged(ESocialSubsystem SubsystemType,
 
 void FSocialUserList::HandleUserGameSpecificStatusChanged(USocialUser* User)
 {
-	// passing dummy Subsystem because HandleUserPresence changed currently doesn't care
 	MarkUserAsDirty(*User);
 	UpdateNow();
 }
@@ -285,7 +286,8 @@ void FSocialUserList::TryAddUserFast(USocialUser& User)
 				// Last step is to check the custom filter, if provided
 				bCanAdd = ListConfig.OnCustomFilterUser.IsBound() ? ListConfig.OnCustomFilterUser.Execute(User) : true;
 
-				// do an initial pass on the GameSpecificStatusFilters (these will only be run again when the user broadcasts OnGameSpecificStatusChanged)
+				// do an initial pass on the GameSpecificStatusFilters on newly added user, and for users that passed other checks but don't yet have the correct presence or game status,
+				// this will be called again from UpdateListInternal(), after the user broadcasts OnGameSpecificStatusChanged, or OnUserPresenceChanged, and the list mark it dirty.
 				for (TFunction<bool(const USocialUser&)> CustomFilterFunction : ListConfig.GameSpecificStatusFilters)
 				{
 					if (!bCanAdd)
@@ -353,7 +355,8 @@ void FSocialUserList::TryRemoveUserFast(USocialUser& User)
 				// We're going to keep the user based on the stock filters, but the custom filter can still veto
 				bRemoveUser = ListConfig.OnCustomFilterUser.IsBound() ? !ListConfig.OnCustomFilterUser.Execute(User) : false;
 
-				// do an initial pass on the GameSpecificStatusFilters (these will only be run again when the user broadcasts OnGameSpecificStatusChanged)
+				// For users that we decide to remove due to incorrect presence or game status, we'll keep listening to their presence and game status change since they may qualify again.
+				// For users that have presence or game status changed but still fit this list, these will be run again next time we got a presence changed or game status changed event.
 				if (!bRemoveUser)
 				{
 					for (TFunction<bool(const USocialUser&)> CustomFilterFunction : ListConfig.GameSpecificStatusFilters)
diff --git a/Engine/Plugins/Online/OnlineFramework/Source/Party/Private/User/SocialUserList.h b/Engine/Plugins/Online/OnlineFramework/Source/Party/Private/User/SocialUserList.h
--- a/Engine/Plugins/Online/OnlineFramework/Source/Party/Private/User/SocialUserList.h
+++ b/Engine/Plugins/Online/OnlineFramework/Source/Party/Private/User/SocialUserList.h
@@ -10,7 +10,7 @@ enum class EMemberExitedReason : uint8;
 class FSocialUserList : public ISocialUserList, public FGCObject, public TSharedFromThis<FSocialUserList>
 {
 public:
-	static TSharedRef<FSocialUserList> CreateUserList(USocialToolkit& InOwnerToolkit, const FSocialUserListConfig& Config);
+	static TSharedRef<FSocialUserList> CreateUserList(const USocialToolkit& InOwnerToolkit, const FSocialUserListConfig& Config);
 
 	virtual void AddReferencedObjects(FReferenceCollector& Collector) override;
 
@@ -27,6 +27,9 @@ public:
 
 	bool HasPresenceFilters() const;
 
+PARTY_SCOPE:
+	const FSocialUserListConfig& GetListConfig() const { return ListConfig; }
+
 private:
 	void HandleOwnerToolkitReset();
 
@@ -67,10 +70,10 @@ private:
 	void HandlePartyMemberLeft(EMemberExitedReason Reason, UPartyMember* Member);
 
 private:
-	FSocialUserList(USocialToolkit& InOwnerToolkit, const FSocialUserListConfig& Config);
+	FSocialUserList(const USocialToolkit& InOwnerToolkit, const FSocialUserListConfig& Config);
 	void InitializeList();
 
-	TWeakObjectPtr<USocialToolkit> OwnerToolkit;
+	TWeakObjectPtr<const USocialToolkit> OwnerToolkit;
 
 	UPROPERTY()
 	TArray<USocialUser*> Users;
@@ -84,7 +87,8 @@ private:
 	FSocialUserListConfig ListConfig;
 
 	bool bNeedsSort = false;
-	float AutoUpdatePeriod = 5.f;
+	int32 AutoUpdateRequests = 0;
+	float AutoUpdatePeriod = .5f;
 	FDelegateHandle UpdateTickerHandle;
 
 	mutable FOnUserAdded OnUserAddedEvent;
diff --git a/Engine/Plugins/Online/OnlineFramework/Source/Party/Public/SocialToolkit.h b/Engine/Plugins/Online/OnlineFramework/Source/Party/Public/SocialToolkit.h
--- a/Engine/Plugins/Online/OnlineFramework/Source/Party/Public/SocialToolkit.h
+++ b/Engine/Plugins/Online/OnlineFramework/Source/Party/Public/SocialToolkit.h
@@ -17,6 +17,7 @@ class FSocialPartyInvitation;
 class USocialChatManager;
 
 class FOnlinePartyId;
+class FSocialUserList;
 
 enum class EMemberExitedReason : uint8;
 
@@ -53,7 +54,7 @@ public:
 
 	IOnlineSubsystem* GetSocialOss(ESocialSubsystem SubsystemType) const;
 
-	TSharedRef<ISocialUserList> CreateUserList(const FSocialUserListConfig& ListConfig);
+	TSharedRef<ISocialUserList> CreateUserList(const FSocialUserListConfig& ListConfig) const;
 
 	USocialUser& GetLocalUser() const;
 
@@ -114,7 +115,7 @@ public:
 	FBasicToolkitEvent& OnToolkitReset() const { return OnToolkitResetEvent; }
 
 #if WITH_EDITOR
-	bool Debug_IsRandomlyChangingPresence() { return bDebug_IsRandomlyChangingUserPresence; }
+	bool Debug_IsRandomlyChangingPresence() const { return bDebug_IsRandomlyChangingUserPresence; }
 #endif
 
 PARTY_SCOPE:
@@ -254,6 +255,8 @@ private:
 	UPROPERTY()
 	USocialChatManager* SocialChatManager;
 
+	mutable TArray<TWeakPtr<FSocialUserList>> CachedSocialUserLists;
+
 	mutable FPartyInviteEvent OnPartyInviteReceivedEvent;
 	mutable FPartyInviteEvent OnPartyInviteAcceptedEvent;
 	mutable FPartyInviteEvent OnPartyInviteRejectedEvent;
diff --git a/Engine/Plugins/Online/OnlineFramework/Source/Party/Public/User/ISocialUserList.h b/Engine/Plugins/Online/OnlineFramework/Source/Party/Public/User/ISocialUserList.h
--- a/Engine/Plugins/Online/OnlineFramework/Source/Party/Public/User/ISocialUserList.h
+++ b/Engine/Plugins/Online/OnlineFramework/Source/Party/Public/User/ISocialUserList.h
@@ -46,6 +46,34 @@ public:
 	bool bAutoUpdate = false;
 
 	bool bSortDuringUpdate = true;
+
+	bool operator==(const FSocialUserListConfig& OtherConfig) const
+	{
+		bool bHasSameConfigs = RelationshipType == OtherConfig.RelationshipType &&
+			RequiredPresenceFlags == OtherConfig.RequiredPresenceFlags &&
+			ForbiddenPresenceFlags == OtherConfig.ForbiddenPresenceFlags &&
+			CompoundSubsystemEnum(RelevantSubsystems) == CompoundSubsystemEnum(OtherConfig.RelevantSubsystems) &&
+			CompoundSubsystemEnum(ForbiddenSubsystems) == CompoundSubsystemEnum(OtherConfig.ForbiddenSubsystems);
+
+		if (bHasSameConfigs && (OnCustomFilterUser.IsBound() || OtherConfig.OnCustomFilterUser.IsBound() || GameSpecificStatusFilters.Num() > 0 || OtherConfig.GameSpecificStatusFilters.Num() > 0))
+		{
+			UE_LOG(LogParty, Verbose, TEXT("Userlist %s and Userlist %s has the exact same list config, but one (or both) have custom filters or game specific filters so we are treating them as separate lists."), *Name, *OtherConfig.Name);
+			return false;
+		}
+
+		return bHasSameConfigs;
+	}
+
+private:
+	uint8 CompoundSubsystemEnum(const TArray<ESocialSubsystem>& Subsystems ) const
+	{
+		uint8 OutputValue = 0;
+		for (ESocialSubsystem Subsystem : Subsystems)
+		{
+			OutputValue |= (uint8)Subsystem;
+		}
+		return OutputValue;
+	}
 };
 
 class ISocialUserList
@@ -64,7 +92,6 @@ public:
 	virtual FOnUpdateComplete& OnUpdateComplete() const = 0;
 
 	virtual const TArray<USocialUser*>& GetUsers() const = 0;
-
 	virtual FString GetListName() const = 0;
 
 	/** Trigger an update of the list immediately, regardless of auto update period */
