Null check for LeftParty in case it's gc'ed after marked pending kill

diff --git a/Engine/Plugins/Online/OnlineFramework/Source/Party/Private/SocialManager.cpp b/Engine/Plugins/Online/OnlineFramework/Source/Party/Private/SocialManager.cpp
--- a/Engine/Plugins/Online/OnlineFramework/Source/Party/Private/SocialManager.cpp
+++ b/Engine/Plugins/Online/OnlineFramework/Source/Party/Private/SocialManager.cpp
@@ -1144,7 +1144,7 @@ void USocialManager::HandlePartyLeft(EMemberExitedReason Reason, USocialParty* L
 			JoinPartyInternal(*JoinAttempt);
 			return;
 		}
-		else
+		else if (LeftParty)
 		{
 			// An attempt to join a party of this type has been initiated, but something/someone decided to leave the party before the attempt was ready to do so
 			// It's not worth accounting for the potential limbo that this could put us into, so just abort the join attempt and let the explicit leave action win
@@ -1153,7 +1153,7 @@ void USocialManager::HandlePartyLeft(EMemberExitedReason Reason, USocialParty* L
 		}
 	}
 
-	if (LeftParty->IsPersistentParty() && GetFirstLocalUserToolkit()->IsOwnerLoggedIn())
+	if (LeftParty && LeftParty->IsPersistentParty() && GetFirstLocalUserToolkit()->IsOwnerLoggedIn())
 	{
 		UE_LOG(LogParty, Verbose, TEXT("Finished leaving persistent party without a join/rejoin target. Creating a new persistent party now."));
 
