Fix shadowed variables

diff --git a/Engine/Plugins/Online/OnlineFramework/Source/Party/Private/SocialManager.cpp b/Engine/Plugins/Online/OnlineFramework/Source/Party/Private/SocialManager.cpp
--- a/Engine/Plugins/Online/OnlineFramework/Source/Party/Private/SocialManager.cpp
+++ b/Engine/Plugins/Online/OnlineFramework/Source/Party/Private/SocialManager.cpp
@@ -700,10 +700,10 @@ void USocialManager::RegisterSecondaryPlayer(int32 LocalPlayerNum, const FOnJoin
 					else
 					{
 						UE_LOG(LogParty, Warning, TEXT("RegisterSecondaryPlayer RemoveLocalMember failed LeaveResult=%s"), ToString(LeaveResult));
-						USocialToolkit* SocialToolkit = GetSocialToolkit(LocalPlayerNum);
-						USocialParty* PersistentParty = GetPersistentParty();
-						FUniqueNetIdRepl SecondaryUserId = SocialToolkit->GetLocalUser().GetUserId(ESocialSubsystem::Primary);
-						JoinDelegate.Execute(*SecondaryUserId, PersistentParty->GetPartyId(), EJoinPartyCompletionResult::UnknownClientFailure, 0);
+						USocialToolkit* LambdaSocialToolkit = GetSocialToolkit(LocalPlayerNum);
+						USocialParty* LambdaPersistentParty = GetPersistentParty();
+						FUniqueNetIdRepl LambdaSecondaryUserId = LambdaSocialToolkit->GetLocalUser().GetUserId(ESocialSubsystem::Primary);
+						JoinDelegate.Execute(*LambdaSecondaryUserId, LambdaPersistentParty->GetPartyId(), EJoinPartyCompletionResult::UnknownClientFailure, 0);
 					}
 				}));
 			}
