Reducing verbosity of SocialUserList sorting lines

#tests Compile DebugGame Editor Win64

diff --git a/Engine/Plugins/Online/OnlineFramework/Source/Party/Private/User/SocialUserList.cpp b/Engine/Plugins/Online/OnlineFramework/Source/Party/Private/User/SocialUserList.cpp
--- a/Engine/Plugins/Online/OnlineFramework/Source/Party/Private/User/SocialUserList.cpp
+++ b/Engine/Plugins/Online/OnlineFramework/Source/Party/Private/User/SocialUserList.cpp
@@ -571,7 +571,7 @@ void FSocialUserList::UpdateListInternal()
 			{
 				SCOPED_NAMED_EVENT(STAT_SocialUserList_Sort, FColor::Orange);
 
-				UE_LOG(LogParty, Verbose, TEXT("%s sorting list of [%d] users"), ANSI_TO_TCHAR(__FUNCTION__), NumUsers);
+				UE_LOG(LogParty, VeryVerbose, TEXT("%s sorting list of [%d] users"), ANSI_TO_TCHAR(__FUNCTION__), NumUsers);
 
 				TArray<FUserSortData> SortedData;
 				SortedData.Reserve(NumUsers);
