Fixing Static Analysis warnings for SocialManager

diff --git a/Engine/Plugins/Online/OnlineFramework/Source/Party/Private/SocialManager.cpp b/Engine/Plugins/Online/OnlineFramework/Source/Party/Private/SocialManager.cpp
--- a/Engine/Plugins/Online/OnlineFramework/Source/Party/Private/SocialManager.cpp
+++ b/Engine/Plugins/Online/OnlineFramework/Source/Party/Private/SocialManager.cpp
@@ -1121,63 +1121,75 @@ void USocialManager::HandlePersistentPartyStateChanged(EPartyState NewState, EPa
 
 void USocialManager::HandleLeavePartyForJoinComplete(ELeavePartyCompletionResult LeaveResult, USocialParty* LeftParty)
 {
-	UE_LOG(LogParty, Verbose, TEXT("Attempt to leave party [%s] for pending join completed with result [%s]"), *LeftParty->ToDebugString(), ToString(LeaveResult));
+	if (LeftParty)
+	{
+		UE_LOG(LogParty, Verbose, TEXT("Attempt to leave party [%s] for pending join completed with result [%s]"), *LeftParty->ToDebugString(), ToString(LeaveResult));
+	}
 }
 
 void USocialManager::HandlePartyDisconnected(USocialParty* DisconnectingParty)
 {
-	const FOnlinePartyTypeId& PartyTypeId = DisconnectingParty->GetPartyTypeId();
-	JoinedPartiesByTypeId.Remove(PartyTypeId);
-	DisconnectingParty->MarkPendingKill();
+	if (DisconnectingParty)
+	{
+		const FOnlinePartyTypeId& PartyTypeId = DisconnectingParty->GetPartyTypeId();
+		JoinedPartiesByTypeId.Remove(PartyTypeId);
+		DisconnectingParty->MarkPendingKill();
+	}
 }
 
 void USocialManager::HandlePartyLeaveBegin(EMemberExitedReason Reason, USocialParty* LeavingParty)
 {
-	const FOnlinePartyTypeId& PartyTypeId = LeavingParty->GetPartyTypeId();
-	JoinedPartiesByTypeId.Remove(PartyTypeId);
-	LeavingPartiesByTypeId.Add(PartyTypeId, LeavingParty);
-}
-
-void USocialManager::HandlePartyLeft(EMemberExitedReason Reason, USocialParty* LeftParty)
-{
-	const FOnlinePartyTypeId& PartyTypeId = LeftParty->GetPartyTypeId();
-	LeavingPartiesByTypeId.Remove(PartyTypeId);
-
-	if (!ensure(!JoinedPartiesByTypeId.Contains(PartyTypeId)))
+	if (LeavingParty)
 	{
-		// Really shouldn't be any scenario wherein we receive a PartyLeft event without a prior PartyLeaveBegin
+		const FOnlinePartyTypeId& PartyTypeId = LeavingParty->GetPartyTypeId();
 		JoinedPartiesByTypeId.Remove(PartyTypeId);
+		LeavingPartiesByTypeId.Add(PartyTypeId, LeavingParty);
 	}
+}
 
-	OnPartyLeftInternal(*LeftParty, Reason);
-	LeftParty->MarkPendingKill();
-
-	if (FJoinPartyAttempt* JoinAttempt = JoinAttemptsByTypeId.Find(PartyTypeId))
+void USocialManager::HandlePartyLeft(EMemberExitedReason Reason, USocialParty* LeftParty)
+{
+	if (LeftParty)
 	{
-		JoinAttempt->ActionTimeTracker.CompleteStep(FJoinPartyAttempt::Step_LeaveCurrentParty);
+		const FOnlinePartyTypeId& PartyTypeId = LeftParty->GetPartyTypeId();
+		LeavingPartiesByTypeId.Remove(PartyTypeId);
 
-		// We're in the process of joining another party of the same type - do we know where we're heading yet?
-		if (JoinAttempt->JoinInfo.IsValid() || JoinAttempt->RejoinInfo.IsValid())
+		if (!ensure(!JoinedPartiesByTypeId.Contains(PartyTypeId)))
 		{
-			// Join the new party immediately and early out
-			JoinPartyInternal(*JoinAttempt);
-			return;
+			// Really shouldn't be any scenario wherein we receive a PartyLeft event without a prior PartyLeaveBegin
+			JoinedPartiesByTypeId.Remove(PartyTypeId);
 		}
-		else if (LeftParty)
+
+		OnPartyLeftInternal(*LeftParty, Reason);
+		LeftParty->MarkPendingKill();
+
+		if (FJoinPartyAttempt* JoinAttempt = JoinAttemptsByTypeId.Find(PartyTypeId))
 		{
-			// An attempt to join a party of this type has been initiated, but something/someone decided to leave the party before the attempt was ready to do so
-			// It's not worth accounting for the potential limbo that this could put us into, so just abort the join attempt and let the explicit leave action win
-			UE_LOG(LogParty, Verbose, TEXT("Finished leaving party [%s] before the current join attempt established join info. Cancelling join attempt."), *LeftParty->ToDebugString());
-			FinishJoinPartyAttempt(*JoinAttempt, FJoinPartyResult(EPartyJoinDenialReason::JoinAttemptAborted));
+			JoinAttempt->ActionTimeTracker.CompleteStep(FJoinPartyAttempt::Step_LeaveCurrentParty);
+
+			// We're in the process of joining another party of the same type - do we know where we're heading yet?
+			if (JoinAttempt->JoinInfo.IsValid() || JoinAttempt->RejoinInfo.IsValid())
+			{
+				// Join the new party immediately and early out
+				JoinPartyInternal(*JoinAttempt);
+				return;
+			}
+			else
+			{
+				// An attempt to join a party of this type has been initiated, but something/someone decided to leave the party before the attempt was ready to do so
+				// It's not worth accounting for the potential limbo that this could put us into, so just abort the join attempt and let the explicit leave action win
+				UE_LOG(LogParty, Verbose, TEXT("Finished leaving party [%s] before the current join attempt established join info. Cancelling join attempt."), *LeftParty->ToDebugString());
+				FinishJoinPartyAttempt(*JoinAttempt, FJoinPartyResult(EPartyJoinDenialReason::JoinAttemptAborted));
+			}
 		}
-	}
 
-	if (LeftParty && LeftParty->IsPersistentParty() && GetFirstLocalUserToolkit()->IsOwnerLoggedIn())
-	{
-		UE_LOG(LogParty, Verbose, TEXT("Finished leaving persistent party without a join/rejoin target. Creating a new persistent party now."));
+		if (LeftParty->IsPersistentParty() && GetFirstLocalUserToolkit()->IsOwnerLoggedIn())
+		{
+			UE_LOG(LogParty, Verbose, TEXT("Finished leaving persistent party without a join/rejoin target. Creating a new persistent party now."));
 
-		// This wasn't part of a join process, so immediately create a new persistent party
-		CreatePersistentParty();
+			// This wasn't part of a join process, so immediately create a new persistent party
+			CreatePersistentParty();
+		}
 	}
 }
 
