Adding logs and quick scope cycle counters to outline hotfix process on perf monitors and log.

diff --git a/Engine/Plugins/Online/OnlineFramework/Source/Hotfix/Private/OnlineHotfixManager.cpp b/Engine/Plugins/Online/OnlineFramework/Source/Hotfix/Private/OnlineHotfixManager.cpp
--- a/Engine/Plugins/Online/OnlineFramework/Source/Hotfix/Private/OnlineHotfixManager.cpp
+++ b/Engine/Plugins/Online/OnlineFramework/Source/Hotfix/Private/OnlineHotfixManager.cpp
@@ -257,6 +257,8 @@ void UOnlineHotfixManager::Cleanup()
 
 void UOnlineHotfixManager::StartHotfixProcess()
 {
+	UE_LOG(LogHotfixManager, Log, TEXT("Starting Hotfix Process"));
+
 	// Patching the editor this way seems like a bad idea
 	bool bShouldHotfix = IsRunningGame() || IsRunningDedicatedServer() || IsRunningClientOnly();
 	if (!bShouldHotfix)
@@ -382,6 +384,8 @@ struct FHotfixFileSortPredicate
 
 void UOnlineHotfixManager::OnEnumerateFilesComplete(bool bWasSuccessful, const FString& ErrorStr)
 {
+	UE_LOG(LogHotfixManager, Log, TEXT("EnumerateFiles Http Request Complete"));
+
 	if (bWasSuccessful)
 	{
 		check(OnlineTitleFile.IsValid());
@@ -751,6 +755,8 @@ bool UOnlineHotfixManager::WantsHotfixProcessing(const FCloudFileHeader& FileHea
 
 bool UOnlineHotfixManager::ApplyHotfixProcessing(const FCloudFileHeader& FileHeader)
 {
+	QUICK_SCOPE_CYCLE_COUNTER(STAT_UOnlineHotfixManager_ApplyHotfixProcessing);
+
 	bool bSuccess = false;
 	const FString Extension = FPaths::GetExtension(FileHeader.FileName);
 	if (Extension == TEXT("INI"))
