Don't attempt to find the same session twice if the find fails - it will very likely continue to fail in perpetuity
TODO Need to identify why find session is failing

diff --git a/Engine/Plugins/Online/OnlineFramework/Source/Party/Private/Party/PartyPlatformSessionMonitor.cpp b/Engine/Plugins/Online/OnlineFramework/Source/Party/Private/Party/PartyPlatformSessionMonitor.cpp
--- a/Engine/Plugins/Online/OnlineFramework/Source/Party/Private/Party/PartyPlatformSessionMonitor.cpp
+++ b/Engine/Plugins/Online/OnlineFramework/Source/Party/Private/Party/PartyPlatformSessionMonitor.cpp
@@ -420,6 +420,10 @@ void FPartyPlatformSessionMonitor::EvaluateCurrentSession()
 			{
 				if (Member->GetRepData().GetPlatformSessionId() == ExistingSessionInfo->SessionId)
 				{
+					if (LastAttemptedFindSessionId.IsSet() && LastAttemptedFindSessionId.GetValue() == ExistingSessionInfo->SessionId)
+					{
+						continue;
+					}
 					// Someone else is claiming to be in the session already, so go find it now
 					FindSession(*ExistingSessionInfo);
 					break;
@@ -525,6 +529,9 @@ void FPartyPlatformSessionMonitor::FindSession(const FPartyPlatformSessionInfo&
 		SessionInitTracker.BeginStep(Step_FindSession);
 		TargetSessionId = SessionInfo.SessionId;
 
+		// Don't attempt to find this session again if this find fails.  This is cleared if the find is successful.
+		LastAttemptedFindSessionId.Emplace(SessionInfo.SessionId);
+
 		SessionManager->FindSession(SessionInfo, FPartyPlatformSessionManager::FOnFindSessionAttemptComplete::CreateSP(this, &FPartyPlatformSessionMonitor::HandleFindSessionComplete));
 	}
 }
@@ -778,6 +785,7 @@ void FPartyPlatformSessionMonitor::HandleFindSessionComplete(bool bWasSuccessful
 	}
 	else if (bWasSuccessful)
 	{
+		LastAttemptedFindSessionId.Reset();
 		JoinSession(FoundSession);
 	}
 	else
diff --git a/Engine/Plugins/Online/OnlineFramework/Source/Party/Private/Party/PartyPlatformSessionMonitor.h b/Engine/Plugins/Online/OnlineFramework/Source/Party/Private/Party/PartyPlatformSessionMonitor.h
--- a/Engine/Plugins/Online/OnlineFramework/Source/Party/Private/Party/PartyPlatformSessionMonitor.h
+++ b/Engine/Plugins/Online/OnlineFramework/Source/Party/Private/Party/PartyPlatformSessionMonitor.h
@@ -112,6 +112,9 @@ private:
 
 	FSessionId TargetSessionId;
 
+	/** Last session id we attempted to find to prevent repeated failures to find the same session */
+	TOptional<FSessionId> LastAttemptedFindSessionId;
+
 	bool bIsSessionMissing = false;
 
 	/** Do we have a console session update queued? */
