- Make configurable whether we set privacy settings to our _desired_ settings (old/default behaviour), or retain the _existing_ settings.

diff --git a/Engine/Plugins/Online/OnlineFramework/Source/Party/Private/Party/SocialParty.cpp b/Engine/Plugins/Online/OnlineFramework/Source/Party/Private/Party/SocialParty.cpp
--- a/Engine/Plugins/Online/OnlineFramework/Source/Party/Private/Party/SocialParty.cpp
+++ b/Engine/Plugins/Online/OnlineFramework/Source/Party/Private/Party/SocialParty.cpp
@@ -520,14 +520,60 @@ FPartyPrivacySettings USocialParty::GetDesiredPrivacySettings() const
 	return FPartyPrivacySettings();
 }
 
+FPartyPrivacySettings USocialParty::GetPrivacySettingsForConfig(const FPartyConfiguration& PartyConfig)
+{
+	FPartyPrivacySettings PrivacySettings;
+
+	// Logic here is a mirror of HandlePrivacySettingsChanged, must update the 2 in tandem
+
+	switch (PartyConfig.PresencePermissions)
+	{
+	case PartySystemPermissions::EPermissionType::Noone:
+		PrivacySettings.PartyType = EPartyType::Private;
+		break;
+	case PartySystemPermissions::EPermissionType::Leader:
+	case PartySystemPermissions::EPermissionType::Friends:
+		PrivacySettings.bOnlyLeaderFriendsCanJoin = true;
+		PrivacySettings.PartyType = EPartyType::Public;
+		break;
+	case PartySystemPermissions::EPermissionType::Anyone:
+		PrivacySettings.bOnlyLeaderFriendsCanJoin = false;
+		PrivacySettings.PartyType = EPartyType::Public;
+		break;
+	}
+
+	switch (PartyConfig.InvitePermissions)
+	{
+	case PartySystemPermissions::EPermissionType::Anyone:
+		PrivacySettings.PartyInviteRestriction = EPartyInviteRestriction::AnyMember;
+		break;
+	case PartySystemPermissions::EPermissionType::Leader:
+	case PartySystemPermissions::EPermissionType::Friends:
+		PrivacySettings.PartyInviteRestriction = EPartyInviteRestriction::LeaderOnly;
+		break;
+	case PartySystemPermissions::EPermissionType::Noone:
+		PrivacySettings.PartyInviteRestriction = EPartyInviteRestriction::NoInvites;
+		break;
+	}
+
+	return PrivacySettings;
+}
+
 void USocialParty::OnLocalPlayerIsLeaderChanged(bool bIsLeader)
 {
 	if (bIsLeader)
 	{
 		GetRepData().OnPrivacySettingsChanged().AddUObject(this, &USocialParty::HandlePrivacySettingsChanged);
 
-		// Establish the privacy of the party to match the local player's preference
-		GetMutableRepData().SetPrivacySettings(GetDesiredPrivacySettings());
+		if (USocialSettings::ShouldSetDesiredPrivacyOnLocalPlayerBecomesLeader())
+		{
+			// Establish the privacy of the party to match the local player's preference
+			GetMutableRepData().SetPrivacySettings(GetDesiredPrivacySettings());
+		}
+		else
+		{
+			GetMutableRepData().SetPrivacySettings(GetPrivacySettingsForConfig(GetCurrentConfiguration()));
+		}
 
 		// It's possible that membership changes resulting in this promotion also require updates to the session info
 		//	If we found out about the changes in membership before learning we're the leader, we were unable to update the rep data accordingly
@@ -950,6 +996,8 @@ void USocialParty::HandlePrivacySettingsChanged(const FPartyPrivacySettings& New
 {
 	check(IsLocalPlayerPartyLeader());
 
+	// Logic here is a mirror of GetPrivacySettingsForConfig, must update the 2 in tandem
+
 	const bool bIsPrivate = NewPrivacySettings.PartyType == EPartyType::Private;
 
 	if (bIsPrivate)
diff --git a/Engine/Plugins/Online/OnlineFramework/Source/Party/Private/SocialSettings.cpp b/Engine/Plugins/Online/OnlineFramework/Source/Party/Private/SocialSettings.cpp
--- a/Engine/Plugins/Online/OnlineFramework/Source/Party/Private/SocialSettings.cpp
+++ b/Engine/Plugins/Online/OnlineFramework/Source/Party/Private/SocialSettings.cpp
@@ -49,6 +49,12 @@ bool USocialSettings::ShouldLeavePartyOnDisconnect()
 	return SettingsCDO.bLeavePartyOnDisconnect;
 }
 
+bool USocialSettings::ShouldSetDesiredPrivacyOnLocalPlayerBecomesLeader()
+{
+	const USocialSettings& SettingsCDO = *GetDefault<USocialSettings>();
+	return SettingsCDO.bSetDesiredPrivacyOnLocalPlayerBecomesLeader;
+}
+
 int32 USocialSettings::GetDefaultMaxPartySize()
 {
 #if !UE_BUILD_SHIPPING
diff --git a/Engine/Plugins/Online/OnlineFramework/Source/Party/Public/Party/SocialParty.h b/Engine/Plugins/Online/OnlineFramework/Source/Party/Public/Party/SocialParty.h
--- a/Engine/Plugins/Online/OnlineFramework/Source/Party/Public/Party/SocialParty.h
+++ b/Engine/Plugins/Online/OnlineFramework/Source/Party/Public/Party/SocialParty.h
@@ -228,6 +228,7 @@ protected:
 	/** Only called when a new party is being created by the local player and they are responsible for the rep data. Otherwise we just wait to receive it from the leader. */
 	virtual void InitializePartyRepData();
 	virtual FPartyPrivacySettings GetDesiredPrivacySettings() const;
+	static FPartyPrivacySettings GetPrivacySettingsForConfig(const FPartyConfiguration& PartyConfig);
 	virtual void OnLocalPlayerIsLeaderChanged(bool bIsLeader);
 	virtual void HandlePrivacySettingsChanged(const FPartyPrivacySettings& NewPrivacySettings);
 	virtual void OnMemberCreatedInternal(UPartyMember& NewMember);
diff --git a/Engine/Plugins/Online/OnlineFramework/Source/Party/Public/SocialSettings.h b/Engine/Plugins/Online/OnlineFramework/Source/Party/Public/SocialSettings.h
--- a/Engine/Plugins/Online/OnlineFramework/Source/Party/Public/SocialSettings.h
+++ b/Engine/Plugins/Online/OnlineFramework/Source/Party/Public/SocialSettings.h
@@ -25,6 +25,7 @@ public:
 	static bool ShouldPreferPlatformInvites();
 	static bool MustSendPrimaryInvites();
 	static bool ShouldLeavePartyOnDisconnect();
+	static bool ShouldSetDesiredPrivacyOnLocalPlayerBecomesLeader();
 	static float GetUserListAutoUpdateRate();
 	static int32 GetMinNicknameLength();
 	static int32 GetMaxNicknameLength();
@@ -53,6 +54,9 @@ private:
 	UPROPERTY(config)
 	bool bLeavePartyOnDisconnect = true;
 
+	UPROPERTY(config)
+	bool bSetDesiredPrivacyOnLocalPlayerBecomesLeader = true;
+
 	UPROPERTY(config)
 	float UserListAutoUpdateRate = 0.5f;
 
