Fix to allow secondary players to be added to platform sessions.

diff --git a/Engine/Plugins/Online/OnlineFramework/Source/Party/Private/Party/PartyPlatformSessionMonitor.cpp b/Engine/Plugins/Online/OnlineFramework/Source/Party/Private/Party/PartyPlatformSessionMonitor.cpp
--- a/Engine/Plugins/Online/OnlineFramework/Source/Party/Private/Party/PartyPlatformSessionMonitor.cpp
+++ b/Engine/Plugins/Online/OnlineFramework/Source/Party/Private/Party/PartyPlatformSessionMonitor.cpp
@@ -397,11 +397,12 @@ void FPartyPlatformSessionMonitor::EvaluateCurrentSession()
 		return;
 	}
 
-	UPartyMember& LocalUserMember = MonitoredParty->GetOwningLocalMember();
 	const IOnlineSessionPtr& SessionInterface = SessionManager->GetSessionInterface();
 	if (const FNamedOnlineSession* Session = SessionInterface->GetNamedSession(PartySessionName))
 	{
+		// TODO: We need to check which of the local players are in the session and process them accordingly 
 		// We already have a platform session session, so we should be all set. Just check the session ids to make sure we have the correct session.
+		UPartyMember& LocalUserMember = MonitoredParty->GetOwningLocalMember();
 		const FSessionId& ReplicatedSessionId = LocalUserMember.GetRepData().GetPlatformSessionId();
 		const FSessionId TrueSessionId = Session->GetSessionIdStr();
 		if (ReplicatedSessionId != TrueSessionId && ensure(DoesLocalUserOwnPlatformSession()))
@@ -430,10 +431,19 @@ void FPartyPlatformSessionMonitor::EvaluateCurrentSession()
 				}
 			}
 		}
-		else if (ExistingSessionInfo->IsSessionOwner(LocalUserMember))
+		else
 		{
-			// There is no session ID yet, but we're the session owner, so it's on us to create it now
-			CreateSession();
+			for (UPartyMember* LocalMember : MonitoredParty->GetPartyMembers())
+			{
+				if (LocalMember->IsLocalPlayer())
+				{
+					if (ExistingSessionInfo->IsSessionOwner(*LocalMember))
+					{
+						// There is no session ID yet, but the session owner is a local player, so it's on us to create it now
+						CreateSession(LocalMember->GetRepData().GetPlatformUniqueId());
+					}
+				}
+			}
 		}
 	}
 	else
@@ -450,6 +460,7 @@ void FPartyPlatformSessionMonitor::Initialize()
 	USocialParty& Party = *MonitoredParty;
 	UE_LOG(LogParty, Verbose, TEXT("Initializing PartyPlatformSessionMonitor for party [%s]"), *Party.GetPartyId().ToDebugString());
 
+	Party.OnPartyMemberCreated().AddSP(this, &FPartyPlatformSessionMonitor::HandlePartyMemberCreated);
 	Party.OnPartyConfigurationChanged().AddSP(this, &FPartyPlatformSessionMonitor::HandlePartyConfigurationChanged);
 	Party.GetRepData().OnPlatformSessionsChanged().AddSP(this, &FPartyPlatformSessionMonitor::HandlePlatformSessionsChanged);
 	Party.OnPartyLeft().AddSP(this, &FPartyPlatformSessionMonitor::HandlePartyLeft);
@@ -471,14 +482,13 @@ void FPartyPlatformSessionMonitor::ShutdownInternal()
 	OnShutdownComplete.ExecuteIfBound(AsShared());
 }
 
-void FPartyPlatformSessionMonitor::CreateSession()
+void FPartyPlatformSessionMonitor::CreateSession(const FUniqueNetIdRepl& LocalUserPlatformId)
 {
 	if (ShutdownState != EMonitorShutdownState::None || !ensure(MonitoredParty.IsValid()))
 	{
 		return;
 	}
 
-	FUniqueNetIdRepl LocalUserPlatformId = MonitoredParty->GetSocialManager().GetFirstLocalUserId(ESocialSubsystem::Platform);
 	FOnlineSessionSettings SessionSettings;
 	if (ensure(LocalUserPlatformId.IsValid()) && ensure(ConfigurePlatformSessionSettings(SessionSettings)))
 	{
@@ -520,6 +530,30 @@ void FPartyPlatformSessionMonitor::CreateSession()
 	}
 }
 
+void FPartyPlatformSessionMonitor::AddLocalPlayerToSession(UPartyMember* PartyMember)
+{
+	const IOnlineSessionPtr& SessionInterface = SessionManager->GetSessionInterface();
+	const FNamedOnlineSession* Session = SessionInterface->GetNamedSession(PartySessionName);
+	check(Session && Session->SessionInfo.IsValid());
+	const FSessionId SessionId = Session->GetSessionIdStr();
+
+	// TODO: Move SetPlatformSessionId to after AddLocalPlayerToSession completes
+	PartyMember->GetMutableRepData().SetPlatformSessionId(SessionId);
+
+	const FUniqueNetIdRepl PartyMemberPlatformUniqueId = PartyMember->GetRepData().GetPlatformUniqueId();
+	if (PartyMemberPlatformUniqueId.IsValid())
+	{
+		SessionInterface->RegisterPlayer(Session->SessionName, *PartyMemberPlatformUniqueId, false);
+	}
+
+	SessionInterface->RegisterLocalPlayer(*PartyMemberPlatformUniqueId, Session->SessionName,
+		FOnRegisterLocalPlayerCompleteDelegate::CreateLambda([](const FUniqueNetId& UserId, EOnJoinSessionCompleteResult::Type Result) 
+			{
+				UE_LOG(LogParty, Log, TEXT("AddLocalPlayerToSession: Complete User=%s Result=%s"), *UserId.ToDebugString(), LexToString(Result));
+			})
+	);
+}
+
 void FPartyPlatformSessionMonitor::FindSession(const FPartyPlatformSessionInfo& SessionInfo)
 {
 	if (ShutdownState == EMonitorShutdownState::None)
@@ -637,12 +671,17 @@ bool FPartyPlatformSessionMonitor::DoesLocalUserOwnPlatformSession()
 		}
 		else
 		{
-			const UPartyMember& LocalMember = MonitoredParty->GetOwningLocalMember();
-			for (const FPartyPlatformSessionInfo& SessionInfo : MonitoredParty->GetRepData().GetPlatformSessions())
+			for (UPartyMember* LocalMember : MonitoredParty->GetPartyMembers())
 			{
-				if (SessionInfo.IsSessionOwner(LocalMember))
+				if (LocalMember->IsLocalPlayer())
 				{
-					return true;
+					for (const FPartyPlatformSessionInfo& SessionInfo : MonitoredParty->GetRepData().GetPlatformSessions())
+					{
+						if (SessionInfo.IsSessionOwner(*LocalMember))
+						{
+							return true;
+						}
+					}
 				}
 			}
 		}
@@ -701,6 +740,12 @@ void FPartyPlatformSessionMonitor::HandlePartyMemberInitialized(UPartyMember* In
 	{
 		SessionManager->GetSessionInterface()->RegisterPlayer(PartySessionName, *InitializedMember->GetRepData().GetPlatformUniqueId(), false);
 	}
+
+	// If a local player joined the party (split screen) we add them to the platform session
+	if (InitializedMember->IsLocalPlayer())
+	{
+		AddLocalPlayerToSession(InitializedMember);
+	}
 }
 
 void FPartyPlatformSessionMonitor::HandlePartyMemberLeft(UPartyMember* OldMember)
@@ -727,21 +772,19 @@ void FPartyPlatformSessionMonitor::HandleCreateSessionComplete(FName SessionName
 		{
 			MonitoredParty->SetIsMissingPlatformSession(false);
 
-			// We successfully created it, we should have it
-			const FNamedOnlineSession* Session = SessionInterface->GetNamedSession(PartySessionName);
-			const FSessionId PlatformSessionId = ensure(Session) ? Session->GetSessionIdStr() : FSessionId();
-			MonitoredParty->GetOwningLocalMember().GetMutableRepData().SetPlatformSessionId(PlatformSessionId);
+			// For all local players
+			for (UPartyMember* PartyMember : MonitoredParty->GetPartyMembers())
+			{
+				if (PartyMember->IsLocalPlayer())
+				{
+					AddLocalPlayerToSession(PartyMember);
+				}
+			}
 
 #if PLATFORM_PS4
 			// Need to queue an immediate update of the newly created session to PUT the ChangeableSessionData
 			QueuePlatformSessionUpdate();
 #endif
-
-			const FUniqueNetIdRepl LocalUserPlatformId = SessionManager->GetLocalUserPlatformId();
-			if (LocalUserPlatformId.IsValid())
-			{
-				SessionInterface->RegisterPlayer(SessionName, *LocalUserPlatformId, false);
-			}
 		}
 		
 		OnSessionEstablished.ExecuteIfBound();
@@ -820,12 +863,17 @@ void FPartyPlatformSessionMonitor::HandleJoinSessionComplete(FName SessionName,
 	else if (bWasSuccessful && ensure(MonitoredParty.IsValid()))
 	{
 		MonitoredParty->SetIsMissingPlatformSession(false);
-		
-		const FUniqueNetIdRepl LocalUserPlatformId = SessionManager->GetLocalUserPlatformId();
-		check(LocalUserPlatformId.IsValid());
-		
-		SessionInterface->RegisterPlayer(SessionName, *LocalUserPlatformId, true);
 
+		// For all local players
+		for (UPartyMember* PartyMember : MonitoredParty->GetPartyMembers())
+		{
+			if (PartyMember->IsLocalPlayer())
+			{
+				AddLocalPlayerToSession(PartyMember);
+			}
+		}
+
+		// For all players only for Tencent
 		if (IsTencentPlatform())
 		{
 			TArray<TSharedRef<const FUniqueNetId>> MemberIdsOnPlatform;
@@ -838,12 +886,6 @@ void FPartyPlatformSessionMonitor::HandleJoinSessionComplete(FName SessionName,
 			}
 			SessionInterface->RegisterPlayers(PartySessionName, MemberIdsOnPlatform);
 		}
-
-		const FNamedOnlineSession* Session = SessionInterface->GetNamedSession(PartySessionName);
-		check(Session && Session->SessionInfo.IsValid());
-		const FSessionId SessionId = Session->GetSessionIdStr();
-		
-		MonitoredParty->GetOwningLocalMember().GetMutableRepData().SetPlatformSessionId(SessionId);
 	}
 	else
 	{
diff --git a/Engine/Plugins/Online/OnlineFramework/Source/Party/Private/Party/PartyPlatformSessionMonitor.h b/Engine/Plugins/Online/OnlineFramework/Source/Party/Private/Party/PartyPlatformSessionMonitor.h
--- a/Engine/Plugins/Online/OnlineFramework/Source/Party/Private/Party/PartyPlatformSessionMonitor.h
+++ b/Engine/Plugins/Online/OnlineFramework/Source/Party/Private/Party/PartyPlatformSessionMonitor.h
@@ -74,7 +74,8 @@ private:
 	void Initialize();
 	void ShutdownInternal();
 
-	void CreateSession();
+	void CreateSession(const FUniqueNetIdRepl& LocalUserPlatformId);
+	void AddLocalPlayerToSession(UPartyMember* InitializedMember);
 	void FindSession(const FPartyPlatformSessionInfo& SessionInfo);
 	void JoinSession(const FOnlineSessionSearchResult& SessionSearchResult);
 	void LeaveSession();
