- Fix party privacy not being correctly set for new parties

diff --git a/Engine/Plugins/Online/OnlineFramework/Source/Party/Private/Party/SocialParty.cpp b/Engine/Plugins/Online/OnlineFramework/Source/Party/Private/Party/SocialParty.cpp
--- a/Engine/Plugins/Online/OnlineFramework/Source/Party/Private/Party/SocialParty.cpp
+++ b/Engine/Plugins/Online/OnlineFramework/Source/Party/Private/Party/SocialParty.cpp
@@ -402,6 +402,12 @@ bool USocialParty::TryKickMember(const UPartyMember& PartyMember)
 	return false;
 }
 
+void USocialParty::ResetPrivacySettings()
+{
+	check(PartyDataReplicator.IsValid());
+	PartyDataReplicator->SetPrivacySettings(GetDesiredPrivacySettings());
+}
+
 const FPartyPrivacySettings& USocialParty::GetPrivacySettings() const
 {
 	check(PartyDataReplicator.IsValid());
diff --git a/Engine/Plugins/Online/OnlineFramework/Source/Party/Private/SocialManager.cpp b/Engine/Plugins/Online/OnlineFramework/Source/Party/Private/SocialManager.cpp
--- a/Engine/Plugins/Online/OnlineFramework/Source/Party/Private/SocialManager.cpp
+++ b/Engine/Plugins/Online/OnlineFramework/Source/Party/Private/SocialManager.cpp
@@ -1014,8 +1014,11 @@ void USocialManager::HandleCreatePartyComplete(const FUniqueNetId& LocalUserId,
 	ECreatePartyCompletionResult LocalCreationResult = Result;
 	if (Result == ECreatePartyCompletionResult::Succeeded)
 	{
-		USocialParty* NewParty = EstablishNewParty(LocalUserId, *PartyId, PartyTypeId);
-		if (!NewParty)
+		if (USocialParty* NewParty = EstablishNewParty(LocalUserId, *PartyId, PartyTypeId))
+		{
+			NewParty->ResetPrivacySettings();
+		}
+		else
 		{
 			LocalCreationResult = ECreatePartyCompletionResult::UnknownClientFailure;
 		}
diff --git a/Engine/Plugins/Online/OnlineFramework/Source/Party/Public/Party/SocialParty.h b/Engine/Plugins/Online/OnlineFramework/Source/Party/Public/Party/SocialParty.h
--- a/Engine/Plugins/Online/OnlineFramework/Source/Party/Public/Party/SocialParty.h
+++ b/Engine/Plugins/Online/OnlineFramework/Source/Party/Public/Party/SocialParty.h
@@ -193,6 +193,7 @@ public:
 	DECLARE_EVENT_TwoParams(USocialParty, FOnPartyMemberConnectionStatusChanged, UPartyMember&, EMemberConnectionStatus);
 	FOnPartyMemberConnectionStatusChanged& OnPartyMemberConnectionStatusChanged() const { return OnPartyMemberConnectionStatusChangedEvent; }
 
+	void ResetPrivacySettings();
 	const FPartyPrivacySettings& GetPrivacySettings() const;
 
 PARTY_SCOPE:
