Move client update checks out of UUpdateManager and into FPatchCheck.
FPatchCheck does not have UObject dependencies which means it can run very early during engine init.

diff --git a/Engine/Plugins/Online/OnlineFramework/OnlineFramework.uplugin b/Engine/Plugins/Online/OnlineFramework/OnlineFramework.uplugin
--- a/Engine/Plugins/Online/OnlineFramework/OnlineFramework.uplugin
+++ b/Engine/Plugins/Online/OnlineFramework/OnlineFramework.uplugin
@@ -23,6 +23,10 @@
 			"Name": "Lobby",
 			"Type": "Runtime"
 		},
+		{
+			"Name": "PatchCheck",
+			"Type": "Runtime"
+		},
 		{
 			"Name": "Hotfix",
 			"Type": "Runtime"
diff --git a/Engine/Plugins/Online/OnlineFramework/Source/Hotfix/Hotfix.Build.cs b/Engine/Plugins/Online/OnlineFramework/Source/Hotfix/Hotfix.Build.cs
--- a/Engine/Plugins/Online/OnlineFramework/Source/Hotfix/Hotfix.Build.cs
+++ b/Engine/Plugins/Online/OnlineFramework/Source/Hotfix/Hotfix.Build.cs
@@ -20,6 +20,12 @@ public class Hotfix : ModuleRules
 			}
 			);
 
+		PublicDependencyModuleNames.AddRange(
+			new string[] {
+				"PatchCheck"
+			}
+			);
+
 		bool bHasOnlineTracing = Directory.Exists(Path.Combine(EngineDirectory, "Plugins", "Online", "NotForLicensees", "OnlineTracing"));
 		if (bHasOnlineTracing)
 		{
diff --git a/Engine/Plugins/Online/OnlineFramework/Source/Hotfix/Private/UpdateManager.cpp b/Engine/Plugins/Online/OnlineFramework/Source/Hotfix/Private/UpdateManager.cpp
--- a/Engine/Plugins/Online/OnlineFramework/Source/Hotfix/Private/UpdateManager.cpp
+++ b/Engine/Plugins/Online/OnlineFramework/Source/Hotfix/Private/UpdateManager.cpp
@@ -53,8 +53,6 @@ UUpdateManager::UUpdateManager()
 	, UpdateCheckCompleteDelay(0.5f)
 	, HotfixAvailabilityCheckCompleteDelay(0.1f)
 	, UpdateCheckAvailabilityCompleteDelay(0.1f)
-	, bCheckPlatformOSSForUpdate(true)
-	, bCheckOSSForUpdate(true)
 	, AppSuspendedUpdateCheckTimeSeconds(600)
 	, bPlatformEnvironmentDetected(false)
 	, bInitialUpdateFinished(false)
@@ -243,114 +241,18 @@ void UUpdateManager::StartPatchCheck()
 {
 	ensure(ChecksEnabled());
 
-	SetUpdateState(EUpdateState::CheckingForPatch);
-	if (bCheckPlatformOSSForUpdate && IOnlineSubsystem::GetByPlatform() != nullptr)
-	{
-		StartPlatformOSSPatchCheck();
-	}
-	else if (bCheckOSSForUpdate)
-	{
-		StartOSSPatchCheck();
-	}
-	else
-	{
-		UE_LOG(LogHotfixManager, Warning, TEXT("Patch check disabled for both Platform and Default OSS"));
-		PatchCheckComplete(EPatchCheckResult::NoPatchRequired);
-	}
-}
-
-void UUpdateManager::StartPlatformOSSPatchCheck()
-{
-	EPatchCheckResult PatchResult = EPatchCheckResult::PatchCheckFailure;
-	bool bStarted = false;
-
-	IOnlineSubsystem* PlatformOnlineSub = IOnlineSubsystem::GetByPlatform();
-	check(PlatformOnlineSub);
-	IOnlineIdentityPtr PlatformOnlineIdentity = PlatformOnlineSub->GetIdentityInterface();
-	if (PlatformOnlineIdentity.IsValid())
-	{
-		TSharedPtr<const FUniqueNetId> UserId = GetFirstSignedInUser(PlatformOnlineIdentity);
-#if PLATFORM_SWITCH
-		// checking the CanPlayOnline privilege on switch will log the user in if required in all but the NotLoggedIn state
-		const bool bCanCheckPlayOnlinePrivilege = UserId.IsValid() && (PlatformOnlineIdentity->GetLoginStatus(*UserId) != ELoginStatus::NotLoggedIn);
-#else
-		const bool bCanCheckPlayOnlinePrivilege = UserId.IsValid() && (PlatformOnlineIdentity->GetLoginStatus(*UserId) == ELoginStatus::LoggedIn);
-#endif
-		if (bCanCheckPlayOnlinePrivilege)
-		{
-			bStarted = true;
-			PlatformOnlineIdentity->GetUserPrivilege(*UserId,
-				EUserPrivileges::CanPlayOnline, IOnlineIdentity::FOnGetUserPrivilegeCompleteDelegate::CreateUObject(this, &ThisClass::OnCheckForPatchComplete, true));
-		}
-		else if (!bInitialUpdateFinished)
-		{
-			UE_LOG(LogHotfixManager, Verbose, TEXT("Skipping initial patch check with no signed in user"));
-			bStarted = true;
-			PatchCheckComplete(EPatchCheckResult::NoPatchRequired);
-		}
-		else
-		{
-			UE_LOG(LogHotfixManager, Warning, TEXT("No valid platform user id when starting patch check!"));
-			PatchResult = EPatchCheckResult::NoLoggedInUser;
-		}
-	}
-
-	if (!bStarted)
-	{
-		// Any failure to call GetUserPrivilege will result in completing the flow via this path
-		PatchCheckComplete(PatchResult);
-	}
-}
-
-void UUpdateManager::StartOSSPatchCheck()
-{
 	UGameInstance* GameInstance = GetGameInstance();
 	check(GameInstance);
 
-	EPatchCheckResult PatchResult = EPatchCheckResult::PatchCheckFailure;
-	bool bStarted = false;
-
+	SetUpdateState(EUpdateState::CheckingForPatch);
 	if (GameInstance->IsDedicatedServerInstance())
 	{
-		bStarted = true;
 		PatchCheckComplete(EPatchCheckResult::NoPatchRequired);
-	}
-	else
-	{
-		UWorld* World = GetWorld();
-		IOnlineIdentityPtr IdentityInt = Online::GetIdentityInterface(World);
-		if (IdentityInt.IsValid())
-		{
-			ULocalPlayer* LP = GameInstance->GetFirstGamePlayer();
-			if (LP != nullptr)
-			{
-				const int32 ControllerId = LP->GetControllerId();
-				TSharedPtr<const FUniqueNetId> UserId = IdentityInt->GetUniquePlayerId(ControllerId);
-				if (!UserId.IsValid())
-				{
-					// Invalid user for "before title/login" check, underlying code doesn't need a valid user currently
-					UserId = IdentityInt->CreateUniquePlayerId(TEXT("InvalidUser"));
-				}
-
-				if (UserId.IsValid())
-				{
-					bStarted = true;
-					IdentityInt->GetUserPrivilege(*UserId,
-						EUserPrivileges::CanPlayOnline, IOnlineIdentity::FOnGetUserPrivilegeCompleteDelegate::CreateUObject(this, &ThisClass::OnCheckForPatchComplete, false));
-				}
-			}
-			else
-			{
-				UE_LOG(LogHotfixManager, Warning, TEXT("No local player to perform check!"));
-			}
-		}
+		return;
 	}
 
-	if (!bStarted)
-	{
-		// Any failure to call GetUserPrivilege will result in completing the flow via this path
-		PatchCheckComplete(PatchResult);
-	}
+	FPatchCheck::Get().GetOnComplete().AddUObject(this, &UUpdateManager::PatchCheckComplete);
+	FPatchCheck::Get().StartPatchCheck();
 }
 
 bool UUpdateManager::ChecksEnabled() const
@@ -358,77 +260,10 @@ bool UUpdateManager::ChecksEnabled() const
 	return !GIsEditor;
 }
 
-bool UUpdateManager::EnvironmentWantsPatchCheck() const
-{
-	return false;
-}
-
-inline bool SkipPatchCheck(UUpdateManager* UpdateManager)
-{
-	// Does the environment care about patch checks (LIVE, STAGE, etc)
-	bool bEnvironmentWantsPatchCheck = UpdateManager->EnvironmentWantsPatchCheck();
-
-	// Can always opt in to a check
-	const bool bForcePatchCheck = FParse::Param(FCommandLine::Get(), TEXT("ForcePatchCheck"));
-
-	// Prevent a patch check on editor builds 
-	const bool bSkipDueToEditor = UE_EDITOR;
-
-	// prevent a check when running unattended
-	const bool bSkipDueToUnattended = FApp::IsUnattended();
-
-	// Explicitly skipping the check
-	const bool bForceSkipCheck = FParse::Param(FCommandLine::Get(), TEXT("SkipPatchCheck"));
-	const bool bSkipPatchCheck = !bForcePatchCheck && (!bEnvironmentWantsPatchCheck || bSkipDueToEditor || bForceSkipCheck || bSkipDueToUnattended);
-
-	return bSkipPatchCheck;
-}
-
-void UUpdateManager::OnCheckForPatchComplete(const FUniqueNetId& UniqueId, EUserPrivileges::Type Privilege, uint32 PrivilegeResult, bool bConsoleCheck)
-{
-	UE_LOG(LogHotfixManager, Verbose, TEXT("[OnCheckForPatchComplete] Privilege=%d PrivilegeResult=%d"), (uint32)Privilege, PrivilegeResult);
-
-	EPatchCheckResult Result = EPatchCheckResult::NoPatchRequired;
-	if (Privilege == EUserPrivileges::CanPlayOnline)
-	{
-		if (bConsoleCheck || !SkipPatchCheck(this))
-		{
-			if (PrivilegeResult & (uint32)IOnlineIdentity::EPrivilegeResults::RequiredSystemUpdate)
-			{
-				Result = EPatchCheckResult::PatchRequired;
-			}
-			else if (PrivilegeResult & (uint32)IOnlineIdentity::EPrivilegeResults::RequiredPatchAvailable)
-			{
-				Result = EPatchCheckResult::PatchRequired;
-			}
-			else if (PrivilegeResult & ((uint32)IOnlineIdentity::EPrivilegeResults::UserNotLoggedIn | (uint32)IOnlineIdentity::EPrivilegeResults::UserNotFound))
-			{
-				Result = EPatchCheckResult::NoLoggedInUser;
-			}
-			else if (PrivilegeResult & (uint32)IOnlineIdentity::EPrivilegeResults::GenericFailure)
-			{
-#if (PLATFORM_XBOXONE || PLATFORM_PS4 || PLATFORM_SWITCH)
-				// Skip console backend failures
-				Result = EPatchCheckResult::NoPatchRequired;
-#else
-				Result = EPatchCheckResult::PatchCheckFailure;
-#endif
-			}
-		}
-	}
-
-	if (bCheckOSSForUpdate && bConsoleCheck && Result == EPatchCheckResult::NoPatchRequired)
-	{
-		// We perform both checks in this case
-		StartOSSPatchCheck();
-		return;
-	}
-
-	PatchCheckComplete(Result);
-}
-
 void UUpdateManager::PatchCheckComplete(EPatchCheckResult PatchResult)
 {
+	FPatchCheck::Get().GetOnComplete().RemoveAll(this);
+
 	LastPatchCheckResult = PatchResult;
 
 	if (PatchResult == EPatchCheckResult::NoPatchRequired)
diff --git a/Engine/Plugins/Online/OnlineFramework/Source/Hotfix/Public/UpdateManager.h b/Engine/Plugins/Online/OnlineFramework/Source/Hotfix/Public/UpdateManager.h
--- a/Engine/Plugins/Online/OnlineFramework/Source/Hotfix/Public/UpdateManager.h
+++ b/Engine/Plugins/Online/OnlineFramework/Source/Hotfix/Public/UpdateManager.h
@@ -8,6 +8,7 @@
 #include "Engine/EngineTypes.h"
 #include "Templates/Casts.h"
 #include "Interfaces/OnlineIdentityInterface.h"
+#include "PatchCheck.h"
 #include "UpdateManager.generated.h"
 
 class Error;
@@ -66,22 +67,6 @@ enum class EUpdateCompletionStatus : uint8
 
 HOTFIX_API FString LexToString(EUpdateCompletionStatus Status);
 
-/**
- * Possible outcomes at the end of just the patch check
- */
-UENUM(BlueprintType)
-enum class EPatchCheckResult : uint8
-{
-	/** No patch required */
-	NoPatchRequired,
-	/** Patch required to continue */
-	PatchRequired,
-	/** Logged in user required for a patch check */
-	NoLoggedInUser,
-	/** Patch check failed */
-	PatchCheckFailure,
-};
-
 /**
  * Delegate fired when changes to the update progress have been made
  *
@@ -239,8 +224,6 @@ protected:
 
 	/** @return true if update checks are enabled */
 	virtual bool ChecksEnabled() const;
-	/** @return true if the backend environment requires update checks */
-	virtual bool EnvironmentWantsPatchCheck() const;
 
 	/** 
 	 * Internal call for StartCheck
@@ -264,22 +247,12 @@ protected:
 	/** Amount of time to wait at the end of the entire check before notifying listening entities (availability check only) */
 	UPROPERTY(Config)
 	float UpdateCheckAvailabilityCompleteDelay;
-	/** Check the platform OSS for an update */
-	UPROPERTY(Config)
-	bool bCheckPlatformOSSForUpdate;
-	/** Check the default OSS for an update */
-	UPROPERTY(Config)
-	bool bCheckOSSForUpdate;
-
 
 	/**
 	 * Patch check
 	 */
 	virtual void StartPatchCheck();
-	virtual void StartPlatformOSSPatchCheck();
-	virtual void StartOSSPatchCheck();
-	virtual void OnCheckForPatchComplete(const FUniqueNetId& UniqueId, EUserPrivileges::Type Privilege, uint32 PrivilegeResult, bool bConsoleCheck);
-	virtual void PatchCheckComplete(EPatchCheckResult PatchResult);
+	void PatchCheckComplete(EPatchCheckResult PatchResult);
 
 	/**
 	 * Hotfix check
@@ -364,8 +337,7 @@ protected:
 	UPROPERTY()
 	int32 WorstNumFilesPendingLoadViewed;
 
-	/** Result of the last patch check */
-	UPROPERTY()
+	/** Result of the last patch check */	
 	EPatchCheckResult LastPatchCheckResult;
 
 	/** Result of the last hotfix */
diff --git a/Engine/Plugins/Online/OnlineFramework/Source/PatchCheck/PatchCheck.build.cs b/Engine/Plugins/Online/OnlineFramework/Source/PatchCheck/PatchCheck.build.cs
new file mode 100644
--- /dev/null
+++ b/Engine/Plugins/Online/OnlineFramework/Source/PatchCheck/PatchCheck.build.cs
@@ -0,0 +1,22 @@
+// Copyright Epic Games, Inc. All Rights Reserved.
+
+using System.IO;
+using UnrealBuildTool;
+
+public class PatchCheck : ModuleRules
+{
+	public PatchCheck(ReadOnlyTargetRules Target) : base(Target)
+	{
+		PCHUsage = ModuleRules.PCHUsageMode.UseExplicitOrSharedPCHs;
+
+		PrivateDependencyModuleNames.AddRange(
+			new string[] {
+				"Core",
+				"Engine",
+				"HTTP",
+				"OnlineSubsystem",
+				"OnlineSubsystemUtils"
+			}
+			);
+	}
+}
diff --git a/Engine/Plugins/Online/OnlineFramework/Source/PatchCheck/Private/PatchCheck.cpp b/Engine/Plugins/Online/OnlineFramework/Source/PatchCheck/Private/PatchCheck.cpp
new file mode 100644
--- /dev/null
+++ b/Engine/Plugins/Online/OnlineFramework/Source/PatchCheck/Private/PatchCheck.cpp
@@ -0,0 +1,215 @@
+// Copyright Epic Games, Inc. All Rights Reserved.
+
+#include "PatchCheck.h"
+#include "PatchCheckModule.h"
+
+#include "Misc/ConfigCacheIni.h"
+#include "Misc/CommandLine.h"
+#include "Misc/App.h"
+#include "Online.h"
+
+DEFINE_LOG_CATEGORY(LogPatchCheck);
+
+FPatchCheck& FPatchCheck::Get()
+{
+	static IPatchCheckModule* ConfiguredModule = nullptr;
+
+	if (ConfiguredModule && ConfiguredModule->GetPatchCheck())
+	{
+		return *ConfiguredModule->GetPatchCheck();
+	}
+
+	FString ModuleName;
+	if (GConfig->GetString(TEXT("PatchCheck"), TEXT("ModuleName"), ModuleName, GEngineIni))
+	{
+		if (FModuleManager::Get().ModuleExists(*ModuleName))
+		{
+			ConfiguredModule = FModuleManager::LoadModulePtr<IPatchCheckModule>(*ModuleName);
+		}
+
+		if (ConfiguredModule && ConfiguredModule->GetPatchCheck())
+		{
+			return *ConfiguredModule->GetPatchCheck();
+		}
+		else
+		{
+			// Couldn't find the configured module, fallback to the default
+			ensureMsgf(false, TEXT("FPatchCheck: Couldn't find module with Name %s, using default"), ModuleName.IsEmpty() ? TEXT("None") : *ModuleName);
+		}
+	}
+
+	// No override module configured, use default
+	ConfiguredModule = &FModuleManager::LoadModuleChecked<IPatchCheckModule>(TEXT("PatchCheck"));
+	return *ConfiguredModule->MakePatchCheck();
+}
+
+FPatchCheck::~FPatchCheck()
+{
+
+}
+
+FPatchCheck::FPatchCheck()
+{
+	GConfig->GetBool(TEXT("PatchCheck"), TEXT("bCheckPlatformOSSForUpdate"), bCheckPlatformOSSForUpdate, GEngineIni);
+	GConfig->GetBool(TEXT("PatchCheck"), TEXT("bCheckOSSForUpdate"), bCheckOSSForUpdate, GEngineIni);
+}
+
+void FPatchCheck::StartPatchCheck()
+{
+	if (bIsCheckInProgress)
+		return;
+
+	if (bCheckPlatformOSSForUpdate && IOnlineSubsystem::GetByPlatform() != nullptr)
+	{
+		bIsCheckInProgress = true;
+		StartPlatformOSSPatchCheck();
+	}
+	else if (bCheckOSSForUpdate)
+	{
+		bIsCheckInProgress = true;
+		StartOSSPatchCheck();
+	}
+	else
+	{
+		UE_LOG(LogPatchCheck, Warning, TEXT("Patch check disabled for both Platform and Default OSS"));
+	}
+}
+
+void FPatchCheck::StartPlatformOSSPatchCheck()
+{
+	EPatchCheckResult PatchResult = EPatchCheckResult::PatchCheckFailure;
+	bool bStarted = false;
+
+	IOnlineSubsystem* PlatformOnlineSub = IOnlineSubsystem::GetByPlatform();
+	check(PlatformOnlineSub);
+	IOnlineIdentityPtr PlatformOnlineIdentity = PlatformOnlineSub->GetIdentityInterface();
+	if (PlatformOnlineIdentity.IsValid())
+	{
+		TSharedPtr<const FUniqueNetId> UserId = GetFirstSignedInUser(PlatformOnlineIdentity);
+#if PLATFORM_SWITCH
+		// checking the CanPlayOnline privilege on switch will log the user in if required in all but the NotLoggedIn state
+		const bool bCanCheckPlayOnlinePrivilege = UserId.IsValid() && (PlatformOnlineIdentity->GetLoginStatus(*UserId) != ELoginStatus::NotLoggedIn);
+#else
+		const bool bCanCheckPlayOnlinePrivilege = UserId.IsValid() && (PlatformOnlineIdentity->GetLoginStatus(*UserId) == ELoginStatus::LoggedIn);
+#endif
+		if (bCanCheckPlayOnlinePrivilege)
+		{
+			bStarted = true;
+			PlatformOnlineIdentity->GetUserPrivilege(*UserId,
+				EUserPrivileges::CanPlayOnline, IOnlineIdentity::FOnGetUserPrivilegeCompleteDelegate::CreateRaw(this, &FPatchCheck::OnCheckForPatchComplete, true));
+		}
+		else
+		{
+			UE_LOG(LogPatchCheck, Warning, TEXT("No valid platform user id when starting patch check!"));
+			PatchResult = EPatchCheckResult::NoLoggedInUser;
+		}
+	}
+
+	if (!bStarted)
+	{
+		// Any failure to call GetUserPrivilege will result in completing the flow via this path
+		PatchCheckComplete(PatchResult);
+	}
+}
+
+void FPatchCheck::StartOSSPatchCheck()
+{
+	EPatchCheckResult PatchResult = EPatchCheckResult::PatchCheckFailure;
+	bool bStarted = false;
+
+	// Online::GetIdentityInterface() can take a UWorld for correctness, but that only matters in PIE right now
+	// and update checks should never happen in PIE currently.
+	IOnlineIdentityPtr IdentityInt = Online::GetIdentityInterface();
+	if (IdentityInt.IsValid())
+	{
+		// User could be invalid for "before title/login" check, underlying code doesn't need a valid user currently
+		TSharedPtr<const FUniqueNetId> UserId = IdentityInt->CreateUniquePlayerId(TEXT("InvalidUser"));
+		if (UserId.IsValid())
+		{
+			bStarted = true;
+			IdentityInt->GetUserPrivilege(*UserId,
+				EUserPrivileges::CanPlayOnline, IOnlineIdentity::FOnGetUserPrivilegeCompleteDelegate::CreateRaw(this, &FPatchCheck::OnCheckForPatchComplete, false));
+		}
+	}
+
+	if (!bStarted)
+	{
+		// Any failure to call GetUserPrivilege will result in completing the flow via this path
+		PatchCheckComplete(PatchResult);
+	}
+}
+
+bool FPatchCheck::EnvironmentWantsPatchCheck() const
+{
+	return false;
+}
+
+bool FPatchCheck::SkipPatchCheck() const
+{
+	// Does the environment care about patch checks (LIVE, STAGE, etc)
+	bool bEnvironmentWantsPatchCheck = EnvironmentWantsPatchCheck();
+
+	// Can always opt in to a check
+	const bool bForcePatchCheck = FParse::Param(FCommandLine::Get(), TEXT("ForcePatchCheck"));
+
+	// Prevent a patch check on editor builds 
+	const bool bSkipDueToEditor = UE_EDITOR;
+
+	// prevent a check when running unattended
+	const bool bSkipDueToUnattended = FApp::IsUnattended();
+
+	// Explicitly skipping the check
+	const bool bForceSkipCheck = FParse::Param(FCommandLine::Get(), TEXT("SkipPatchCheck"));
+	const bool bSkipPatchCheck = !bForcePatchCheck && (!bEnvironmentWantsPatchCheck || bSkipDueToEditor || bForceSkipCheck || bSkipDueToUnattended);
+
+	return bSkipPatchCheck;
+}
+
+void FPatchCheck::OnCheckForPatchComplete(const FUniqueNetId& UniqueId, EUserPrivileges::Type Privilege, uint32 PrivilegeResult, bool bConsoleCheck)
+{
+	UE_LOG(LogPatchCheck, Verbose, TEXT("[OnCheckForPatchComplete] Privilege=%d PrivilegeResult=%d"), (uint32)Privilege, PrivilegeResult);
+
+	EPatchCheckResult Result = EPatchCheckResult::NoPatchRequired;
+	if (Privilege == EUserPrivileges::CanPlayOnline)
+	{
+		if (bConsoleCheck || !SkipPatchCheck())
+		{
+			if (PrivilegeResult & (uint32)IOnlineIdentity::EPrivilegeResults::RequiredSystemUpdate)
+			{
+				Result = EPatchCheckResult::PatchRequired;
+			}
+			else if (PrivilegeResult & (uint32)IOnlineIdentity::EPrivilegeResults::RequiredPatchAvailable)
+			{
+				Result = EPatchCheckResult::PatchRequired;
+			}
+			else if (PrivilegeResult & ((uint32)IOnlineIdentity::EPrivilegeResults::UserNotLoggedIn | (uint32)IOnlineIdentity::EPrivilegeResults::UserNotFound))
+			{
+				Result = EPatchCheckResult::NoLoggedInUser;
+			}
+			else if (PrivilegeResult & (uint32)IOnlineIdentity::EPrivilegeResults::GenericFailure)
+			{
+#if (PLATFORM_XBOXONE || PLATFORM_PS4 || PLATFORM_SWITCH)
+				// Skip console backend failures
+				Result = EPatchCheckResult::NoPatchRequired;
+#else
+				Result = EPatchCheckResult::PatchCheckFailure;
+#endif
+			}
+		}
+	}
+
+	if (bCheckOSSForUpdate && bConsoleCheck && Result == EPatchCheckResult::NoPatchRequired)
+	{
+		// We perform both checks in this case
+		StartOSSPatchCheck();
+		return;
+	}
+
+	PatchCheckComplete(Result);
+}
+
+void FPatchCheck::PatchCheckComplete(EPatchCheckResult PatchResult)
+{
+	OnComplete.Broadcast(PatchResult);
+	bIsCheckInProgress = false;
+}
diff --git a/Engine/Plugins/Online/OnlineFramework/Source/PatchCheck/Private/PatchCheckModule.cpp b/Engine/Plugins/Online/OnlineFramework/Source/PatchCheck/Private/PatchCheckModule.cpp
new file mode 100644
--- /dev/null
+++ b/Engine/Plugins/Online/OnlineFramework/Source/PatchCheck/Private/PatchCheckModule.cpp
@@ -0,0 +1,10 @@
+// Copyright Epic Games, Inc. All Rights Reserved.
+
+#include "PatchCheckModule.h"
+#include "PatchCheck.h"
+
+class FPatchCheckModule : public TPatchCheckModule<FPatchCheck>
+{
+};
+
+IMPLEMENT_MODULE(FPatchCheckModule, PatchCheck);
diff --git a/Engine/Plugins/Online/OnlineFramework/Source/PatchCheck/Public/PatchCheck.h b/Engine/Plugins/Online/OnlineFramework/Source/PatchCheck/Public/PatchCheck.h
new file mode 100644
--- /dev/null
+++ b/Engine/Plugins/Online/OnlineFramework/Source/PatchCheck/Public/PatchCheck.h
@@ -0,0 +1,68 @@
+// Copyright Epic Games, Inc. All Rights Reserved.
+
+#pragma once
+
+#include "PatchCheckModule.h"
+#include "Interfaces/OnlineIdentityInterface.h"
+
+PATCHCHECK_API DECLARE_LOG_CATEGORY_EXTERN(LogPatchCheck, Log, All);
+
+/**
+ * Possible outcomes at the end of just the patch check
+ */
+enum class EPatchCheckResult : uint8
+{
+	/** No patch required */
+	NoPatchRequired,
+	/** Patch required to continue */
+	PatchRequired,
+	/** Logged in user required for a patch check */
+	NoLoggedInUser,
+	/** Patch check failed */
+	PatchCheckFailure,
+	Count,
+};
+
+DECLARE_MULTICAST_DELEGATE_OneParam(FOnPatchCheckComplete, EPatchCheckResult /*Result*/);
+
+class PATCHCHECK_API FPatchCheck
+{
+public:
+	static FPatchCheck& Get();
+
+	virtual ~FPatchCheck();
+
+protected:
+	FPatchCheck();
+
+private:
+	FPatchCheck(const FPatchCheck& Other) = delete;
+	FPatchCheck& operator=(const FPatchCheck& Other) = delete;
+
+public:
+	void StartPatchCheck();
+
+	FOnPatchCheckComplete& GetOnComplete() { return OnComplete; }
+
+protected:
+	virtual void StartPlatformOSSPatchCheck();
+	virtual void StartOSSPatchCheck();
+	virtual bool EnvironmentWantsPatchCheck() const;
+	bool SkipPatchCheck() const;
+	void OnCheckForPatchComplete(const FUniqueNetId& UniqueId, EUserPrivileges::Type Privilege, uint32 PrivilegeResult, bool bConsoleCheck);
+	virtual void PatchCheckComplete(EPatchCheckResult PatchResult);
+
+protected:
+
+	FOnPatchCheckComplete OnComplete;
+
+	/** Track whether we can start a new check */
+	bool bIsCheckInProgress = false;
+
+	/** Check the platform OSS for an update */	
+	bool bCheckPlatformOSSForUpdate = true;
+	/** Check the default OSS for an update */
+	bool bCheckOSSForUpdate = true;
+
+	friend class TPatchCheckModule<FPatchCheck>;
+};
diff --git a/Engine/Plugins/Online/OnlineFramework/Source/PatchCheck/Public/PatchCheckModule.h b/Engine/Plugins/Online/OnlineFramework/Source/PatchCheck/Public/PatchCheckModule.h
new file mode 100644
--- /dev/null
+++ b/Engine/Plugins/Online/OnlineFramework/Source/PatchCheck/Public/PatchCheckModule.h
@@ -0,0 +1,56 @@
+// Copyright Epic Games, Inc. All Rights Reserved.
+
+// Copyright Epic Games, Inc. All Rights Reserved.
+
+#pragma once
+
+#include "Modules/ModuleInterface.h"
+#include "Modules/ModuleManager.h"
+#include "Misc/ConfigCacheIni.h"
+
+class FPatchCheck;
+
+class IPatchCheckModule : public IModuleInterface
+{
+public:
+	virtual void StartupModule() override
+	{
+		// Only instantiate if this is the version the game has been configured to use
+		FString ModuleName;
+		GConfig->GetString(TEXT("PatchCheck"), TEXT("ModuleName"), ModuleName, GEngineIni);
+
+		if (FModuleManager::Get().GetModule(*ModuleName) == this)
+		{
+			MakePatchCheck();
+		}
+	}
+
+	virtual void PreUnloadCallback() override
+	{
+		PatchCheck = nullptr;
+	}
+
+	virtual FPatchCheck* MakePatchCheck() = 0;
+
+	FPatchCheck* GetPatchCheck()
+	{
+		return PatchCheck.Get();
+	}
+
+protected:
+	TUniquePtr<FPatchCheck> PatchCheck;
+};
+
+template<class PatchCheckClass>
+class TPatchCheckModule : public IPatchCheckModule
+{
+public:
+	virtual FPatchCheck* MakePatchCheck() override
+	{
+		if (PatchCheck == nullptr)
+		{
+			PatchCheck = TUniquePtr<PatchCheckClass>(new PatchCheckClass());
+		}
+		return PatchCheck.Get();
+	}
+};
