Route patch checks from UUpdateManager through install bundle manager if enabled.  This will allow for content patch checks to occur.

diff --git a/Engine/Plugins/Online/OnlineFramework/Source/Hotfix/Hotfix.Build.cs b/Engine/Plugins/Online/OnlineFramework/Source/Hotfix/Hotfix.Build.cs
--- a/Engine/Plugins/Online/OnlineFramework/Source/Hotfix/Hotfix.Build.cs
+++ b/Engine/Plugins/Online/OnlineFramework/Source/Hotfix/Hotfix.Build.cs
@@ -22,7 +22,8 @@ public class Hotfix : ModuleRules
 
 		PublicDependencyModuleNames.AddRange(
 			new string[] {
-				"PatchCheck"
+				"PatchCheck",
+				"InstallBundleManager",
 			}
 			);
 
diff --git a/Engine/Plugins/Online/OnlineFramework/Source/Hotfix/Private/UpdateManager.cpp b/Engine/Plugins/Online/OnlineFramework/Source/Hotfix/Private/UpdateManager.cpp
--- a/Engine/Plugins/Online/OnlineFramework/Source/Hotfix/Private/UpdateManager.cpp
+++ b/Engine/Plugins/Online/OnlineFramework/Source/Hotfix/Private/UpdateManager.cpp
@@ -10,6 +10,8 @@
 #include "Engine/LocalPlayer.h"
 #include "OnlineSubsystem.h"
 #include "Misc/CoreDelegates.h"
+#include "InstallBundleManagerInterface.h"
+#include "InstallBundleUtils.h"
 
 #include "OnlineHotfixManager.h"
 
@@ -59,7 +61,7 @@ UUpdateManager::UUpdateManager()
 	, bCheckHotfixAvailabilityOnly(false)
 	, CurrentUpdateState(EUpdateState::UpdateIdle)
 	, WorstNumFilesPendingLoadViewed(0)
-	, LastPatchCheckResult(EPatchCheckResult::PatchCheckFailure)
+	, LastPatchCheckResult(EInstallBundleManagerPatchCheckResult::PatchCheckFailure)
 	, LastHotfixResult(EHotfixResult::Failed)
 	, LoadStartTime(0.0)
 	, UpdateStateEnum(nullptr)
@@ -251,8 +253,17 @@ void UUpdateManager::StartPatchCheck()
 		return;
 	}
 
-	FPatchCheck::Get().GetOnComplete().AddUObject(this, &UUpdateManager::PatchCheckComplete);
-	FPatchCheck::Get().StartPatchCheck();
+	IInstallBundleManager* InstallBundleMan = IInstallBundleManager::GetPlatformInstallBundleManager();
+	if (InstallBundleMan && !InstallBundleMan->IsNullInterface())
+	{
+		InstallBundleMan->PatchCheckCompleteDelegate.AddUObject(this, &UUpdateManager::InstallBundlePatchCheckComplete);
+		InstallBundleMan->StartPatchCheck();
+	}
+	else
+	{
+		FPatchCheck::Get().GetOnComplete().AddUObject(this, &UUpdateManager::PatchCheckComplete);
+		FPatchCheck::Get().StartPatchCheck();
+	}
 }
 
 bool UUpdateManager::ChecksEnabled() const
@@ -260,24 +271,63 @@ bool UUpdateManager::ChecksEnabled() const
 	return !GIsEditor;
 }
 
+EInstallBundleManagerPatchCheckResult ToInstallBundleManagerPatchCheckResult(EPatchCheckResult InResult)
+{
+	// EInstallBundleManagerPatchCheckResult should be a superset of EPatchCheckResult
+
+	EInstallBundleManagerPatchCheckResult OutResult = EInstallBundleManagerPatchCheckResult::PatchCheckFailure;
+	switch (InResult)
+	{
+	case EPatchCheckResult::NoPatchRequired:
+		OutResult = EInstallBundleManagerPatchCheckResult::NoPatchRequired;
+		break;
+	case EPatchCheckResult::PatchRequired:
+		OutResult = EInstallBundleManagerPatchCheckResult::ClientPatchRequired;
+		break;
+	case EPatchCheckResult::NoLoggedInUser:
+		OutResult = EInstallBundleManagerPatchCheckResult::NoLoggedInUser;
+		break;
+	case EPatchCheckResult::PatchCheckFailure:
+		OutResult = EInstallBundleManagerPatchCheckResult::PatchCheckFailure;
+		break;
+	default:
+		ensureAlwaysMsgf(false, TEXT("Unknown EPatchCheckResult"));
+		OutResult = EInstallBundleManagerPatchCheckResult::PatchCheckFailure;
+		break;
+	}
+
+	// Make sure we don't miss a case
+	static_assert(InstallBundleUtil::CastToUnderlying(EPatchCheckResult::Count) == 4, "");
+
+	return OutResult;
+}
+
 void UUpdateManager::PatchCheckComplete(EPatchCheckResult PatchResult)
 {
 	FPatchCheck::Get().GetOnComplete().RemoveAll(this);
 
+	InstallBundlePatchCheckComplete(ToInstallBundleManagerPatchCheckResult(PatchResult));
+}
+
+void UUpdateManager::InstallBundlePatchCheckComplete(EInstallBundleManagerPatchCheckResult PatchResult)
+{
+	IInstallBundleManager::PatchCheckCompleteDelegate.RemoveAll(this);
+
 	LastPatchCheckResult = PatchResult;
 
-	if (PatchResult == EPatchCheckResult::NoPatchRequired)
+	if (PatchResult == EInstallBundleManagerPatchCheckResult::NoPatchRequired)
 	{
 		StartPlatformEnvironmentCheck();
 	}
-	else if (PatchResult == EPatchCheckResult::NoLoggedInUser)
+	else if (PatchResult == EInstallBundleManagerPatchCheckResult::NoLoggedInUser)
 	{
 		CheckComplete(EUpdateCompletionStatus::UpdateFailure_NotLoggedIn);
 	}
 	else
 	{
-		ensure(PatchResult == EPatchCheckResult::PatchCheckFailure ||
-			   PatchResult == EPatchCheckResult::PatchRequired);
+		ensure(PatchResult == EInstallBundleManagerPatchCheckResult::PatchCheckFailure ||
+			   PatchResult == EInstallBundleManagerPatchCheckResult::ClientPatchRequired ||
+			   PatchResult == EInstallBundleManagerPatchCheckResult::ContentPatchRequired);
 
 		// Skip hotfix check in error states, but still preload data as if there was nothing wrong
 		StartInitialPreload();
@@ -501,17 +551,21 @@ void UUpdateManager::InitialPreloadComplete()
 {
 	SetUpdateState(EUpdateState::InitialLoadComplete);
 
-	if (LastPatchCheckResult == EPatchCheckResult::PatchCheckFailure)
+	if (LastPatchCheckResult == EInstallBundleManagerPatchCheckResult::PatchCheckFailure)
 	{
 		CheckComplete(EUpdateCompletionStatus::UpdateFailure_PatchCheck);
 	}
-	else if (LastPatchCheckResult == EPatchCheckResult::PatchRequired)
+	else if (LastPatchCheckResult == EInstallBundleManagerPatchCheckResult::ClientPatchRequired)
 	{
 		CheckComplete(EUpdateCompletionStatus::UpdateSuccess_NeedsPatch);
 	}
+	else if (LastPatchCheckResult == EInstallBundleManagerPatchCheckResult::ContentPatchRequired)
+	{
+		CheckComplete(EUpdateCompletionStatus::UpdateSuccess_NeedsRelaunch);
+	}
 	else
 	{
-		ensure(LastPatchCheckResult == EPatchCheckResult::NoPatchRequired);
+		ensure(LastPatchCheckResult == EInstallBundleManagerPatchCheckResult::NoPatchRequired);
 		// Patch check success, check hotfix status
 		switch (LastHotfixResult)
 		{
diff --git a/Engine/Plugins/Online/OnlineFramework/Source/Hotfix/Public/UpdateManager.h b/Engine/Plugins/Online/OnlineFramework/Source/Hotfix/Public/UpdateManager.h
--- a/Engine/Plugins/Online/OnlineFramework/Source/Hotfix/Public/UpdateManager.h
+++ b/Engine/Plugins/Online/OnlineFramework/Source/Hotfix/Public/UpdateManager.h
@@ -9,6 +9,7 @@
 #include "Templates/Casts.h"
 #include "Interfaces/OnlineIdentityInterface.h"
 #include "PatchCheck.h"
+#include "InstallBundleTypes.h"
 #include "UpdateManager.generated.h"
 
 class Error;
@@ -252,6 +253,7 @@ protected:
 	 * Patch check
 	 */
 	virtual void StartPatchCheck();
+	void InstallBundlePatchCheckComplete(EInstallBundleManagerPatchCheckResult PatchResult);
 	void PatchCheckComplete(EPatchCheckResult PatchResult);
 
 	/**
@@ -338,7 +340,7 @@ protected:
 	int32 WorstNumFilesPendingLoadViewed;
 
 	/** Result of the last patch check */	
-	EPatchCheckResult LastPatchCheckResult;
+	EInstallBundleManagerPatchCheckResult LastPatchCheckResult;
 
 	/** Result of the last hotfix */
 	UPROPERTY()
