Add backwards compatibility layer to Patch check refactor.

diff --git a/Engine/Plugins/Online/OnlineFramework/Source/Hotfix/Private/UpdateManager.cpp b/Engine/Plugins/Online/OnlineFramework/Source/Hotfix/Private/UpdateManager.cpp
--- a/Engine/Plugins/Online/OnlineFramework/Source/Hotfix/Private/UpdateManager.cpp
+++ b/Engine/Plugins/Online/OnlineFramework/Source/Hotfix/Private/UpdateManager.cpp
@@ -239,6 +239,11 @@ void UUpdateManager::CheckComplete(EUpdateCompletionStatus Result, bool bUpdateT
 	DelayResponse(CompletionDelegate, bCheckHotfixAvailabilityOnly ? UpdateCheckAvailabilityCompleteDelay : UpdateCheckCompleteDelay);
 }
 
+inline FName GetUniqueTag(UUpdateManager* UpdateManager)
+{
+	return FName(*FString::Printf(TEXT("Tag_%u_%s"), UpdateManager->GetUniqueID(), *UpdateManager->GetName()));
+}
+
 void UUpdateManager::StartPatchCheck()
 {
 	ensure(ChecksEnabled());
@@ -257,11 +262,17 @@ void UUpdateManager::StartPatchCheck()
 	if (InstallBundleMan && !InstallBundleMan->IsNullInterface())
 	{
 		InstallBundleMan->PatchCheckCompleteDelegate.AddUObject(this, &UUpdateManager::InstallBundlePatchCheckComplete);
+		InstallBundleMan->AddEnvironmentWantsPatchCheckBackBackCompatDelegate(
+			GetUniqueTag(this),
+			FEnvironmentWantsPatchCheck::CreateUObject(this, &UUpdateManager::EnvironmentWantsPatchCheck));
 		InstallBundleMan->StartPatchCheck();
 	}
 	else
 	{
 		FPatchCheck::Get().GetOnComplete().AddUObject(this, &UUpdateManager::PatchCheckComplete);
+		FPatchCheck::Get().AddEnvironmentWantsPatchCheckBackBackCompatDelegate(
+			GetUniqueTag(this),
+			FEnvironmentWantsPatchCheck::CreateUObject(this, &UUpdateManager::EnvironmentWantsPatchCheck));
 		FPatchCheck::Get().StartPatchCheck();
 	}
 }
@@ -271,6 +282,11 @@ bool UUpdateManager::ChecksEnabled() const
 	return !GIsEditor;
 }
 
+bool UUpdateManager::EnvironmentWantsPatchCheck() const
+{
+	return false;
+}
+
 EInstallBundleManagerPatchCheckResult ToInstallBundleManagerPatchCheckResult(EPatchCheckResult InResult)
 {
 	// EInstallBundleManagerPatchCheckResult should be a superset of EPatchCheckResult
@@ -305,13 +321,18 @@ EInstallBundleManagerPatchCheckResult ToInstallBundleManagerPatchCheckResult(EPa
 void UUpdateManager::PatchCheckComplete(EPatchCheckResult PatchResult)
 {
 	FPatchCheck::Get().GetOnComplete().RemoveAll(this);
+	FPatchCheck::Get().RemoveEnvironmentWantsPatchCheckBackBackCompatDelegate(GetUniqueTag(this));
 
 	InstallBundlePatchCheckComplete(ToInstallBundleManagerPatchCheckResult(PatchResult));
 }
 
 void UUpdateManager::InstallBundlePatchCheckComplete(EInstallBundleManagerPatchCheckResult PatchResult)
 {
+	IInstallBundleManager* InstallBundleMan = IInstallBundleManager::GetPlatformInstallBundleManager();
+	check(InstallBundleMan && !InstallBundleMan->IsNullInterface());
+
 	IInstallBundleManager::PatchCheckCompleteDelegate.RemoveAll(this);
+	InstallBundleMan->RemoveEnvironmentWantsPatchCheckBackBackCompatDelegate(GetUniqueTag(this));
 
 	LastPatchCheckResult = PatchResult;
 
diff --git a/Engine/Plugins/Online/OnlineFramework/Source/Hotfix/Public/UpdateManager.h b/Engine/Plugins/Online/OnlineFramework/Source/Hotfix/Public/UpdateManager.h
--- a/Engine/Plugins/Online/OnlineFramework/Source/Hotfix/Public/UpdateManager.h
+++ b/Engine/Plugins/Online/OnlineFramework/Source/Hotfix/Public/UpdateManager.h
@@ -225,6 +225,9 @@ protected:
 
 	/** @return true if update checks are enabled */
 	virtual bool ChecksEnabled() const;
+	/** @return true if the backend environment requires update checks */
+	UE_DEPRECATED_FORGAME(4.25, "UUpdateManager::EnvironmentWantsPatchCheck is deprecated, please override FPatchCheck::EnvironmentWantsPatchCheck instead.")
+	virtual bool EnvironmentWantsPatchCheck() const;
 
 	/** 
 	 * Internal call for StartCheck
@@ -249,6 +252,13 @@ protected:
 	UPROPERTY(Config)
 	float UpdateCheckAvailabilityCompleteDelay;
 
+	/** Check the platform OSS for an update */
+	UE_DEPRECATED(4.25, "Set FPatchCheck::bCheckPlatformOSSForUpdate using section [PatchCheck] instead.")
+	bool bCheckPlatformOSSForUpdate;
+	/** Check the default OSS for an update */
+	UE_DEPRECATED(4.25, "Set FPatchCheck::bCheckOSSForUpdate using section [PatchCheck] instead.")
+	bool bCheckOSSForUpdate;
+
 	/**
 	 * Patch check
 	 */
diff --git a/Engine/Plugins/Online/OnlineFramework/Source/PatchCheck/Private/PatchCheck.cpp b/Engine/Plugins/Online/OnlineFramework/Source/PatchCheck/Private/PatchCheck.cpp
--- a/Engine/Plugins/Online/OnlineFramework/Source/PatchCheck/Private/PatchCheck.cpp
+++ b/Engine/Plugins/Online/OnlineFramework/Source/PatchCheck/Private/PatchCheck.cpp
@@ -50,8 +50,23 @@ FPatchCheck::~FPatchCheck()
 
 FPatchCheck::FPatchCheck()
 {
-	GConfig->GetBool(TEXT("PatchCheck"), TEXT("bCheckPlatformOSSForUpdate"), bCheckPlatformOSSForUpdate, GEngineIni);
-	GConfig->GetBool(TEXT("PatchCheck"), TEXT("bCheckOSSForUpdate"), bCheckOSSForUpdate, GEngineIni);
+	if (!GConfig->GetBool(TEXT("PatchCheck"), TEXT("bCheckPlatformOSSForUpdate"), bCheckPlatformOSSForUpdate, GEngineIni))
+	{
+		/** For backwards compatibility with UUpdateManager */
+		if (GConfig->GetBool(TEXT("/Script/Hotfix.UpdateManager"), TEXT("bCheckPlatformOSSForUpdate"), bCheckPlatformOSSForUpdate, GEngineIni))
+		{
+			ensureMsgf(false, TEXT("UpdateManager::bCheckPlatformOSSForUpdate is deprecated, Set FPatchCheck::bCheckPlatformOSSForUpdate using section [PatchCheck] instead."));
+		}
+	}
+
+	if (!GConfig->GetBool(TEXT("PatchCheck"), TEXT("bCheckOSSForUpdate"), bCheckOSSForUpdate, GEngineIni))
+	{
+		/** For backwards compatibility with UUpdateManager */
+		if (GConfig->GetBool(TEXT("/Script/Hotfix.UpdateManager"), TEXT("bCheckOSSForUpdate"), bCheckOSSForUpdate, GEngineIni))
+		{
+			ensureMsgf(false, TEXT("UpdateManager::bCheckOSSForUpdate is deprecated, Set FPatchCheck::bCheckOSSForUpdate using section [PatchCheck] instead."));
+		}
+	}
 }
 
 void FPatchCheck::StartPatchCheck()
@@ -75,6 +90,16 @@ void FPatchCheck::StartPatchCheck()
 	}
 }
 
+void FPatchCheck::AddEnvironmentWantsPatchCheckBackBackCompatDelegate(FName Tag, FEnvironmentWantsPatchCheck Delegate)
+{
+	BackCompatEnvironmentWantsPatchCheckDelegates.Emplace(Tag, MoveTemp(Delegate));
+}
+
+void FPatchCheck::RemoveEnvironmentWantsPatchCheckBackBackCompatDelegate(FName Tag)
+{
+	BackCompatEnvironmentWantsPatchCheckDelegates.Remove(Tag);
+}
+
 void FPatchCheck::StartPlatformOSSPatchCheck()
 {
 	EPatchCheckResult PatchResult = EPatchCheckResult::PatchCheckFailure;
@@ -141,6 +166,17 @@ void FPatchCheck::StartOSSPatchCheck()
 
 bool FPatchCheck::EnvironmentWantsPatchCheck() const
 {
+	if (BackCompatEnvironmentWantsPatchCheckDelegates.Num() > 0)
+	{
+		for (const TPair<FName, FEnvironmentWantsPatchCheck>& Pair : BackCompatEnvironmentWantsPatchCheckDelegates)
+		{
+			if (Pair.Value.IsBound() && Pair.Value.Execute())
+			{
+				return true;
+			}
+		}
+	}
+
 	return false;
 }
 
diff --git a/Engine/Plugins/Online/OnlineFramework/Source/PatchCheck/Public/PatchCheck.h b/Engine/Plugins/Online/OnlineFramework/Source/PatchCheck/Public/PatchCheck.h
--- a/Engine/Plugins/Online/OnlineFramework/Source/PatchCheck/Public/PatchCheck.h
+++ b/Engine/Plugins/Online/OnlineFramework/Source/PatchCheck/Public/PatchCheck.h
@@ -25,6 +25,9 @@ enum class EPatchCheckResult : uint8
 
 DECLARE_MULTICAST_DELEGATE_OneParam(FOnPatchCheckComplete, EPatchCheckResult /*Result*/);
 
+// For backwards compatibility only!
+DECLARE_DELEGATE_RetVal(bool, FEnvironmentWantsPatchCheck);
+
 class PATCHCHECK_API FPatchCheck
 {
 public:
@@ -42,6 +45,9 @@ private:
 public:
 	void StartPatchCheck();
 
+	void AddEnvironmentWantsPatchCheckBackBackCompatDelegate(FName Tag, FEnvironmentWantsPatchCheck Delegate);
+	void RemoveEnvironmentWantsPatchCheckBackBackCompatDelegate(FName Tag);
+
 	FOnPatchCheckComplete& GetOnComplete() { return OnComplete; }
 
 protected:
@@ -56,6 +62,9 @@ protected:
 
 	FOnPatchCheckComplete OnComplete;
 
+	/** For backwards compatibility with UUpdateManager */
+	TMap<FName, FEnvironmentWantsPatchCheck> BackCompatEnvironmentWantsPatchCheckDelegates;
+
 	/** Track whether we can start a new check */
 	bool bIsCheckInProgress = false;
 
diff --git a/Engine/Source/Runtime/InstallBundleManager/Public/InstallBundleManagerInterface.h b/Engine/Source/Runtime/InstallBundleManager/Public/InstallBundleManagerInterface.h
--- a/Engine/Source/Runtime/InstallBundleManager/Public/InstallBundleManagerInterface.h
+++ b/Engine/Source/Runtime/InstallBundleManager/Public/InstallBundleManagerInterface.h
@@ -55,6 +55,8 @@ DECLARE_MULTICAST_DELEGATE_OneParam(FInstallBundlePausedMultiDelegate, FInstallB
 
 DECLARE_MULTICAST_DELEGATE_OneParam(FInstallBundleManagerOnPatchCheckComplete, EInstallBundleManagerPatchCheckResult);
 
+DECLARE_DELEGATE_RetVal(bool, FInstallBundleManagerEnvironmentWantsPatchCheck);
+
 class INSTALLBUNDLEMANAGER_API IInstallBundleManager
 {
 public:
@@ -108,6 +110,8 @@ public:
 	virtual void UpdateContentRequestFlags(TArrayView<const FName> BundleNames, EInstallBundleRequestFlags AddFlags, EInstallBundleRequestFlags RemoveFlags) = 0;
 	
 	virtual void StartPatchCheck();
+	virtual void AddEnvironmentWantsPatchCheckBackBackCompatDelegate(FName Tag, FInstallBundleManagerEnvironmentWantsPatchCheck Delegate) {}
+	virtual void RemoveEnvironmentWantsPatchCheckBackBackCompatDelegate(FName Tag) {}
 
 	virtual bool IsNullInterface() const = 0;
 
