#UE4 Change check to an if, per Justin Marcus

diff --git a/Engine/Plugins/Online/OnlineFramework/Source/Hotfix/Private/UpdateManager.cpp b/Engine/Plugins/Online/OnlineFramework/Source/Hotfix/Private/UpdateManager.cpp
--- a/Engine/Plugins/Online/OnlineFramework/Source/Hotfix/Private/UpdateManager.cpp
+++ b/Engine/Plugins/Online/OnlineFramework/Source/Hotfix/Private/UpdateManager.cpp
@@ -329,29 +329,30 @@ void UUpdateManager::PatchCheckComplete(EPatchCheckResult PatchResult)
 void UUpdateManager::InstallBundlePatchCheckComplete(EInstallBundleManagerPatchCheckResult PatchResult)
 {
 	IInstallBundleManager* InstallBundleMan = IInstallBundleManager::GetPlatformInstallBundleManager();
-	check(InstallBundleMan && !InstallBundleMan->IsNullInterface());
-
-	IInstallBundleManager::PatchCheckCompleteDelegate.RemoveAll(this);
-	InstallBundleMan->RemoveEnvironmentWantsPatchCheckBackBackCompatDelegate(GetUniqueTag(this));
+	if (InstallBundleMan && !InstallBundleMan->IsNullInterface())
+	{
+		IInstallBundleManager::PatchCheckCompleteDelegate.RemoveAll(this);
+		InstallBundleMan->RemoveEnvironmentWantsPatchCheckBackBackCompatDelegate(GetUniqueTag(this));
 
-	LastPatchCheckResult = PatchResult;
+		LastPatchCheckResult = PatchResult;
 
-	if (PatchResult == EInstallBundleManagerPatchCheckResult::NoPatchRequired)
-	{
-		StartPlatformEnvironmentCheck();
-	}
-	else if (PatchResult == EInstallBundleManagerPatchCheckResult::NoLoggedInUser)
-	{
-		CheckComplete(EUpdateCompletionStatus::UpdateFailure_NotLoggedIn);
-	}
-	else
-	{
-		ensure(PatchResult == EInstallBundleManagerPatchCheckResult::PatchCheckFailure ||
-			   PatchResult == EInstallBundleManagerPatchCheckResult::ClientPatchRequired ||
-			   PatchResult == EInstallBundleManagerPatchCheckResult::ContentPatchRequired);
+		if (PatchResult == EInstallBundleManagerPatchCheckResult::NoPatchRequired)
+		{
+			StartPlatformEnvironmentCheck();
+		}
+		else if (PatchResult == EInstallBundleManagerPatchCheckResult::NoLoggedInUser)
+		{
+			CheckComplete(EUpdateCompletionStatus::UpdateFailure_NotLoggedIn);
+		}
+		else
+		{
+			ensure(PatchResult == EInstallBundleManagerPatchCheckResult::PatchCheckFailure ||
+				   PatchResult == EInstallBundleManagerPatchCheckResult::ClientPatchRequired ||
+				   PatchResult == EInstallBundleManagerPatchCheckResult::ContentPatchRequired);
 
-		// Skip hotfix check in error states, but still preload data as if there was nothing wrong
-		StartInitialPreload();
+			// Skip hotfix check in error states, but still preload data as if there was nothing wrong
+			StartInitialPreload();
+		}
 	}
 }
 
