Fix void UUpdateManager::InstallBundlePatchCheckComplete(EInstallBundleManagerPatchCheckResult PatchResult)

diff --git a/Engine/Plugins/Online/OnlineFramework/Source/Hotfix/Private/UpdateManager.cpp b/Engine/Plugins/Online/OnlineFramework/Source/Hotfix/Private/UpdateManager.cpp
--- a/Engine/Plugins/Online/OnlineFramework/Source/Hotfix/Private/UpdateManager.cpp
+++ b/Engine/Plugins/Online/OnlineFramework/Source/Hotfix/Private/UpdateManager.cpp
@@ -331,28 +331,28 @@ void UUpdateManager::InstallBundlePatchCheckComplete(EInstallBundleManagerPatchC
 	IInstallBundleManager* InstallBundleMan = IInstallBundleManager::GetPlatformInstallBundleManager();
 	if (InstallBundleMan && !InstallBundleMan->IsNullInterface())
 	{
-		IInstallBundleManager::PatchCheckCompleteDelegate.RemoveAll(this);
 		InstallBundleMan->RemoveEnvironmentWantsPatchCheckBackBackCompatDelegate(GetUniqueTag(this));
+	}
+	IInstallBundleManager::PatchCheckCompleteDelegate.RemoveAll(this);
 
-		LastPatchCheckResult = PatchResult;
+	LastPatchCheckResult = PatchResult;
 
-		if (PatchResult == EInstallBundleManagerPatchCheckResult::NoPatchRequired)
-		{
-			StartPlatformEnvironmentCheck();
-		}
-		else if (PatchResult == EInstallBundleManagerPatchCheckResult::NoLoggedInUser)
-		{
-			CheckComplete(EUpdateCompletionStatus::UpdateFailure_NotLoggedIn);
-		}
-		else
-		{
-			ensure(PatchResult == EInstallBundleManagerPatchCheckResult::PatchCheckFailure ||
-				   PatchResult == EInstallBundleManagerPatchCheckResult::ClientPatchRequired ||
-				   PatchResult == EInstallBundleManagerPatchCheckResult::ContentPatchRequired);
+	if (PatchResult == EInstallBundleManagerPatchCheckResult::NoPatchRequired)
+	{
+		StartPlatformEnvironmentCheck();
+	}
+	else if (PatchResult == EInstallBundleManagerPatchCheckResult::NoLoggedInUser)
+	{
+		CheckComplete(EUpdateCompletionStatus::UpdateFailure_NotLoggedIn);
+	}
+	else
+	{
+		ensure(PatchResult == EInstallBundleManagerPatchCheckResult::PatchCheckFailure ||
+			PatchResult == EInstallBundleManagerPatchCheckResult::ClientPatchRequired ||
+			PatchResult == EInstallBundleManagerPatchCheckResult::ContentPatchRequired);
 
-			// Skip hotfix check in error states, but still preload data as if there was nothing wrong
-			StartInitialPreload();
-		}
+		// Skip hotfix check in error states, but still preload data as if there was nothing wrong
+		StartInitialPreload();
 	}
 }
 
