Fix typos

diff --git a/Engine/Plugins/Online/OnlineFramework/Source/Hotfix/Private/UpdateManager.cpp b/Engine/Plugins/Online/OnlineFramework/Source/Hotfix/Private/UpdateManager.cpp
--- a/Engine/Plugins/Online/OnlineFramework/Source/Hotfix/Private/UpdateManager.cpp
+++ b/Engine/Plugins/Online/OnlineFramework/Source/Hotfix/Private/UpdateManager.cpp
@@ -262,7 +262,7 @@ void UUpdateManager::StartPatchCheck()
 	if (InstallBundleMan && !InstallBundleMan->IsNullInterface())
 	{
 		InstallBundleMan->PatchCheckCompleteDelegate.AddUObject(this, &UUpdateManager::InstallBundlePatchCheckComplete);
-		InstallBundleMan->AddEnvironmentWantsPatchCheckBackBackCompatDelegate(
+		InstallBundleMan->AddEnvironmentWantsPatchCheckBackCompatDelegate(
 			GetUniqueTag(this),
 			FEnvironmentWantsPatchCheck::CreateUObject(this, &UUpdateManager::EnvironmentWantsPatchCheck));
 		InstallBundleMan->StartPatchCheck();
@@ -270,7 +270,7 @@ void UUpdateManager::StartPatchCheck()
 	else
 	{
 		FPatchCheck::Get().GetOnComplete().AddUObject(this, &UUpdateManager::PatchCheckComplete);
-		FPatchCheck::Get().AddEnvironmentWantsPatchCheckBackBackCompatDelegate(
+		FPatchCheck::Get().AddEnvironmentWantsPatchCheckBackCompatDelegate(
 			GetUniqueTag(this),
 			FEnvironmentWantsPatchCheck::CreateUObject(this, &UUpdateManager::EnvironmentWantsPatchCheck));
 		FPatchCheck::Get().StartPatchCheck();
@@ -321,7 +321,7 @@ EInstallBundleManagerPatchCheckResult ToInstallBundleManagerPatchCheckResult(EPa
 void UUpdateManager::PatchCheckComplete(EPatchCheckResult PatchResult)
 {
 	FPatchCheck::Get().GetOnComplete().RemoveAll(this);
-	FPatchCheck::Get().RemoveEnvironmentWantsPatchCheckBackBackCompatDelegate(GetUniqueTag(this));
+	FPatchCheck::Get().RemoveEnvironmentWantsPatchCheckBackCompatDelegate(GetUniqueTag(this));
 
 	InstallBundlePatchCheckComplete(ToInstallBundleManagerPatchCheckResult(PatchResult));
 }
@@ -331,7 +331,7 @@ void UUpdateManager::InstallBundlePatchCheckComplete(EInstallBundleManagerPatchC
 	IInstallBundleManager* InstallBundleMan = IInstallBundleManager::GetPlatformInstallBundleManager();
 	if (InstallBundleMan && !InstallBundleMan->IsNullInterface())
 	{
-		InstallBundleMan->RemoveEnvironmentWantsPatchCheckBackBackCompatDelegate(GetUniqueTag(this));
+		InstallBundleMan->RemoveEnvironmentWantsPatchCheckBackCompatDelegate(GetUniqueTag(this));
 	}
 	IInstallBundleManager::PatchCheckCompleteDelegate.RemoveAll(this);
 
diff --git a/Engine/Plugins/Online/OnlineFramework/Source/PatchCheck/Private/PatchCheck.cpp b/Engine/Plugins/Online/OnlineFramework/Source/PatchCheck/Private/PatchCheck.cpp
--- a/Engine/Plugins/Online/OnlineFramework/Source/PatchCheck/Private/PatchCheck.cpp
+++ b/Engine/Plugins/Online/OnlineFramework/Source/PatchCheck/Private/PatchCheck.cpp
@@ -90,12 +90,12 @@ void FPatchCheck::StartPatchCheck()
 	}
 }
 
-void FPatchCheck::AddEnvironmentWantsPatchCheckBackBackCompatDelegate(FName Tag, FEnvironmentWantsPatchCheck Delegate)
+void FPatchCheck::AddEnvironmentWantsPatchCheckBackCompatDelegate(FName Tag, FEnvironmentWantsPatchCheck Delegate)
 {
 	BackCompatEnvironmentWantsPatchCheckDelegates.Emplace(Tag, MoveTemp(Delegate));
 }
 
-void FPatchCheck::RemoveEnvironmentWantsPatchCheckBackBackCompatDelegate(FName Tag)
+void FPatchCheck::RemoveEnvironmentWantsPatchCheckBackCompatDelegate(FName Tag)
 {
 	BackCompatEnvironmentWantsPatchCheckDelegates.Remove(Tag);
 }
diff --git a/Engine/Plugins/Online/OnlineFramework/Source/PatchCheck/Public/PatchCheck.h b/Engine/Plugins/Online/OnlineFramework/Source/PatchCheck/Public/PatchCheck.h
--- a/Engine/Plugins/Online/OnlineFramework/Source/PatchCheck/Public/PatchCheck.h
+++ b/Engine/Plugins/Online/OnlineFramework/Source/PatchCheck/Public/PatchCheck.h
@@ -40,13 +40,14 @@ protected:
 
 private:
 	FPatchCheck(const FPatchCheck& Other) = delete;
+
 	FPatchCheck& operator=(const FPatchCheck& Other) = delete;
 
 public:
 	void StartPatchCheck();
 
-	void AddEnvironmentWantsPatchCheckBackBackCompatDelegate(FName Tag, FEnvironmentWantsPatchCheck Delegate);
-	void RemoveEnvironmentWantsPatchCheckBackBackCompatDelegate(FName Tag);
+	void AddEnvironmentWantsPatchCheckBackCompatDelegate(FName Tag, FEnvironmentWantsPatchCheck Delegate);
+	void RemoveEnvironmentWantsPatchCheckBackCompatDelegate(FName Tag);
 
 	FOnPatchCheckComplete& GetOnComplete() { return OnComplete; }
 
diff --git a/Engine/Source/Runtime/InstallBundleManager/Public/InstallBundleManagerInterface.h b/Engine/Source/Runtime/InstallBundleManager/Public/InstallBundleManagerInterface.h
--- a/Engine/Source/Runtime/InstallBundleManager/Public/InstallBundleManagerInterface.h
+++ b/Engine/Source/Runtime/InstallBundleManager/Public/InstallBundleManagerInterface.h
@@ -110,8 +110,8 @@ public:
 	virtual void UpdateContentRequestFlags(TArrayView<const FName> BundleNames, EInstallBundleRequestFlags AddFlags, EInstallBundleRequestFlags RemoveFlags) = 0;
 	
 	virtual void StartPatchCheck();
-	virtual void AddEnvironmentWantsPatchCheckBackBackCompatDelegate(FName Tag, FInstallBundleManagerEnvironmentWantsPatchCheck Delegate) {}
-	virtual void RemoveEnvironmentWantsPatchCheckBackBackCompatDelegate(FName Tag) {}
+	virtual void AddEnvironmentWantsPatchCheckBackCompatDelegate(FName Tag, FInstallBundleManagerEnvironmentWantsPatchCheck Delegate) {}
+	virtual void RemoveEnvironmentWantsPatchCheckBackCompatDelegate(FName Tag) {}
 
 	virtual bool IsNullInterface() const = 0;
 
