- Fix SessionInfo not being created for platform members

diff --git a/Engine/Plugins/Online/OnlineFramework/Source/Party/Private/Party/SocialParty.cpp b/Engine/Plugins/Online/OnlineFramework/Source/Party/Private/Party/SocialParty.cpp
--- a/Engine/Plugins/Online/OnlineFramework/Source/Party/Private/Party/SocialParty.cpp
+++ b/Engine/Plugins/Online/OnlineFramework/Source/Party/Private/Party/SocialParty.cpp
@@ -924,18 +924,24 @@ void USocialParty::HandleMemberInitialized(UPartyMember* Member)
 {
 	if (IsLocalPlayerPartyLeader())
 	{
+		Member->GetRepData().OnPlatformUniqueIdChanged().AddUObject(this, &USocialParty::HandleMemberPlatformUniqueIdChanged, Member);
 		Member->GetRepData().OnPlatformSessionIdChanged().AddUObject(this, &USocialParty::HandleMemberSessionIdChanged, Member);
 
-		const FName MemberPlatformOssName = Member->GetPlatformOssName();
-		if (FPartyPlatformSessionManager::DoesOssNeedPartySession(MemberPlatformOssName) && !GetRepData().FindSessionInfo(MemberPlatformOssName))
-		{
-			// We don't have session info yet for this platform, so make it now and establish this member as the owner
-			FPartyPlatformSessionInfo NewSessionInfo;
-			NewSessionInfo.OssName = MemberPlatformOssName;
-			NewSessionInfo.OwnerPrimaryId = Member->GetPrimaryNetId();
+		HandleMemberPlatformUniqueIdChanged(Member->GetRepData().GetPlatformUniqueId(), Member);
+	}
+}
 
-			GetMutableRepData().UpdatePlatformSessionInfo(NewSessionInfo);
-		}
+void USocialParty::HandleMemberPlatformUniqueIdChanged(const FUniqueNetIdRepl& NewPlatformUniqueId, UPartyMember* Member)
+{
+	const FName MemberPlatformOssName = NewPlatformUniqueId.GetType();
+	if (FPartyPlatformSessionManager::DoesOssNeedPartySession(MemberPlatformOssName) && !GetRepData().FindSessionInfo(MemberPlatformOssName))
+	{
+		// We don't have session info yet for this platform, so make it now and establish this member as the owner
+		FPartyPlatformSessionInfo NewSessionInfo;
+		NewSessionInfo.OssName = MemberPlatformOssName;
+		NewSessionInfo.OwnerPrimaryId = Member->GetPrimaryNetId();
+
+		GetMutableRepData().UpdatePlatformSessionInfo(NewSessionInfo);
 	}
 }
 
diff --git a/Engine/Plugins/Online/OnlineFramework/Source/Party/Public/Party/SocialParty.h b/Engine/Plugins/Online/OnlineFramework/Source/Party/Public/Party/SocialParty.h
--- a/Engine/Plugins/Online/OnlineFramework/Source/Party/Public/Party/SocialParty.h
+++ b/Engine/Plugins/Online/OnlineFramework/Source/Party/Public/Party/SocialParty.h
@@ -327,6 +327,7 @@ private:	// Handlers
 	void HandlePartyPromotionLockoutChanged(const FUniqueNetId& LocalUserId, const FOnlinePartyId& PartyId, bool bArePromotionsLocked);
 
 	void HandleMemberInitialized(UPartyMember* Member);
+	void HandleMemberPlatformUniqueIdChanged(const FUniqueNetIdRepl& NewPlatformUniqueId, UPartyMember* Member);
 	void HandleMemberSessionIdChanged(const FSessionId& NewSessionId, UPartyMember* Member);
 
 	void HandleBeaconHostConnectionFailed();
