#social
Fix for static analysis warning

diff --git a/Engine/Plugins/Online/OnlineFramework/Source/Party/Private/SocialManager.cpp b/Engine/Plugins/Online/OnlineFramework/Source/Party/Private/SocialManager.cpp
--- a/Engine/Plugins/Online/OnlineFramework/Source/Party/Private/SocialManager.cpp
+++ b/Engine/Plugins/Online/OnlineFramework/Source/Party/Private/SocialManager.cpp
@@ -1140,6 +1140,11 @@ void USocialManager::HandlePartyLeaveBegin(EMemberExitedReason Reason, USocialPa
 
 void USocialManager::HandlePartyLeft(EMemberExitedReason Reason, USocialParty* LeftParty)
 {
+	if (!ensure(LeftParty))
+	{
+		return;
+	}
+
 	const FOnlinePartyTypeId& PartyTypeId = LeftParty->GetPartyTypeId();
 	LeavingPartiesByTypeId.Remove(PartyTypeId);
 
@@ -1163,7 +1168,7 @@ void USocialManager::HandlePartyLeft(EMemberExitedReason Reason, USocialParty* L
 			JoinPartyInternal(*JoinAttempt);
 			return;
 		}
-		else if (LeftParty)
+		else
 		{
 			// An attempt to join a party of this type has been initiated, but something/someone decided to leave the party before the attempt was ready to do so
 			// It's not worth accounting for the potential limbo that this could put us into, so just abort the join attempt and let the explicit leave action win
@@ -1172,7 +1177,7 @@ void USocialManager::HandlePartyLeft(EMemberExitedReason Reason, USocialParty* L
 		}
 	}
 
-	if (LeftParty && LeftParty->IsPersistentParty() && GetFirstLocalUserToolkit()->IsOwnerLoggedIn())
+	if (LeftParty->IsPersistentParty() && GetFirstLocalUserToolkit()->IsOwnerLoggedIn())
 	{
 		UE_LOG(LogParty, Verbose, TEXT("Finished leaving persistent party without a join/rejoin target. Creating a new persistent party now."));
 
