Fix not clearing platform sessions that have no members when we are promoted

diff --git a/Engine/Plugins/Online/OnlineFramework/Source/Party/Private/Party/SocialParty.cpp b/Engine/Plugins/Online/OnlineFramework/Source/Party/Private/Party/SocialParty.cpp
--- a/Engine/Plugins/Online/OnlineFramework/Source/Party/Private/Party/SocialParty.cpp
+++ b/Engine/Plugins/Online/OnlineFramework/Source/Party/Private/Party/SocialParty.cpp
@@ -593,6 +593,16 @@ void USocialParty::OnLocalPlayerIsLeaderChanged(bool bIsLeader)
 				AllMemberPlatformSubsystems.AddUnique(PlatformOssName);
 			}
 		}
+		for (const FPartyPlatformSessionInfo& PlatformSessionInfo : GetRepData().GetPlatformSessions())
+		{
+			// Not using AddUnique so we can log if we are catching OSSs to remove
+			// TODO: Remove logging when we are sure this is fixed
+			if (!AllMemberPlatformSubsystems.Contains(PlatformSessionInfo.OssName))
+			{
+				UE_LOG(LogParty, Verbose, TEXT("OnLocalPlayerIsLeaderChanged: Adding existing platform OSS [%s] which has no members"), *PlatformSessionInfo.OssName.ToString());
+				AllMemberPlatformSubsystems.Add(PlatformSessionInfo.OssName);
+			}
+		}
 		for (FName PlatformOssName : AllMemberPlatformSubsystems)
 		{
 			UpdatePlatformSessionLeader(PlatformOssName);
