- Moved the engine plugins that were under NotForLicensees and NoRedist to Restricted
- Moved the VerseCompilerCmd binaries to Engine/Restricted/NotForLicensees/Binaries

diff --git a/Engine/Plugins/Online/OnlineFramework/Source/Hotfix/Hotfix.Build.cs b/Engine/Plugins/Online/OnlineFramework/Source/Hotfix/Hotfix.Build.cs
--- a/Engine/Plugins/Online/OnlineFramework/Source/Hotfix/Hotfix.Build.cs
+++ b/Engine/Plugins/Online/OnlineFramework/Source/Hotfix/Hotfix.Build.cs
@@ -27,7 +27,7 @@ public class Hotfix : ModuleRules
 			}
 			);
 
-		bool bHasOnlineTracing = Directory.Exists(Path.Combine(EngineDirectory, "Plugins", "Online", "NotForLicensees", "OnlineTracing"));
+		bool bHasOnlineTracing = Directory.Exists(Path.Combine(EngineDirectory, "Restricted", "NotForLicensees", "Plugins", "Online", "OnlineTracing"));
 		if (bHasOnlineTracing)
 		{
 			PublicDefinitions.Add("WITH_ONLINETRACING=1");
