- Fix crashes caused by adding local players to non existent party session

diff --git a/Engine/Plugins/Online/OnlineFramework/Source/Party/Private/Party/PartyPlatformSessionMonitor.cpp b/Engine/Plugins/Online/OnlineFramework/Source/Party/Private/Party/PartyPlatformSessionMonitor.cpp
--- a/Engine/Plugins/Online/OnlineFramework/Source/Party/Private/Party/PartyPlatformSessionMonitor.cpp
+++ b/Engine/Plugins/Online/OnlineFramework/Source/Party/Private/Party/PartyPlatformSessionMonitor.cpp
@@ -433,14 +433,15 @@ void FPartyPlatformSessionMonitor::EvaluateCurrentSession()
 		}
 		else
 		{
-			for (UPartyMember* LocalMember : MonitoredParty->GetPartyMembers())
+			for (UPartyMember* Member : MonitoredParty->GetPartyMembers())
 			{
-				if (LocalMember->IsLocalPlayer())
+				if (Member->IsLocalPlayer())
 				{
-					if (ExistingSessionInfo->IsSessionOwner(*LocalMember))
+					if (ExistingSessionInfo->IsSessionOwner(*Member))
 					{
 						// There is no session ID yet, but the session owner is a local player, so it's on us to create it now
-						CreateSession(LocalMember->GetRepData().GetPlatformUniqueId());
+						CreateSession(Member->GetRepData().GetPlatformUniqueId());
+						break;
 					}
 				}
 			}
@@ -532,26 +533,31 @@ void FPartyPlatformSessionMonitor::CreateSession(const FUniqueNetIdRepl& LocalUs
 
 void FPartyPlatformSessionMonitor::AddLocalPlayerToSession(UPartyMember* PartyMember)
 {
-	const IOnlineSessionPtr& SessionInterface = SessionManager->GetSessionInterface();
-	const FNamedOnlineSession* Session = SessionInterface->GetNamedSession(PartySessionName);
-	check(Session && Session->SessionInfo.IsValid());
-	const FSessionId SessionId = Session->GetSessionIdStr();
-
-	// TODO: Move SetPlatformSessionId to after AddLocalPlayerToSession completes
-	PartyMember->GetMutableRepData().SetPlatformSessionId(SessionId);
-
-	const FUniqueNetIdRepl PartyMemberPlatformUniqueId = PartyMember->GetRepData().GetPlatformUniqueId();
-	if (PartyMemberPlatformUniqueId.IsValid())
+	if (ensure(MonitoredParty.IsValid()) &&
+		!MonitoredParty->IsMissingPlatformSession())
 	{
-		SessionInterface->RegisterPlayer(Session->SessionName, *PartyMemberPlatformUniqueId, false);
-	}
+		const IOnlineSessionPtr& SessionInterface = SessionManager->GetSessionInterface();
+		const FNamedOnlineSession* Session = SessionInterface->GetNamedSession(PartySessionName);
+		if (ensure(Session && Session->SessionInfo.IsValid()))
+		{
+			const FSessionId SessionId = Session->GetSessionIdStr();
 
-	SessionInterface->RegisterLocalPlayer(*PartyMemberPlatformUniqueId, Session->SessionName,
-		FOnRegisterLocalPlayerCompleteDelegate::CreateLambda([](const FUniqueNetId& UserId, EOnJoinSessionCompleteResult::Type Result) 
+			// TODO: Move SetPlatformSessionId to after AddLocalPlayerToSession completes
+			PartyMember->GetMutableRepData().SetPlatformSessionId(SessionId);
+
+			const FUniqueNetIdRepl PartyMemberPlatformUniqueId = PartyMember->GetRepData().GetPlatformUniqueId();
+			if (PartyMemberPlatformUniqueId.IsValid())
 			{
-				UE_LOG(LogParty, Log, TEXT("AddLocalPlayerToSession: Complete User=%s Result=%s"), *UserId.ToDebugString(), LexToString(Result));
-			})
-	);
+				SessionInterface->RegisterPlayer(Session->SessionName, *PartyMemberPlatformUniqueId, false);
+
+				SessionInterface->RegisterLocalPlayer(*PartyMemberPlatformUniqueId, Session->SessionName,
+					FOnRegisterLocalPlayerCompleteDelegate::CreateLambda([](const FUniqueNetId& UserId, EOnJoinSessionCompleteResult::Type Result)
+				{
+					UE_LOG(LogParty, Log, TEXT("AddLocalPlayerToSession: Complete User=%s Result=%s"), *UserId.ToDebugString(), LexToString(Result));
+				}));
+			}
+		}
+	}
 }
 
 void FPartyPlatformSessionMonitor::FindSession(const FPartyPlatformSessionInfo& SessionInfo)
@@ -631,7 +637,11 @@ void FPartyPlatformSessionMonitor::LeaveSession()
 {
 	UE_LOG(LogParty, Log, TEXT("PartyPlatformSessionMonitor destroying platform party session now."));
 	ShutdownState = EMonitorShutdownState::InProgress;
-	SessionManager->GetSessionInterface()->DestroySession(PartySessionName, FOnDestroySessionCompleteDelegate::CreateSP(this, &FPartyPlatformSessionMonitor::HandleDestroySessionComplete));
+
+	const IOnlineSessionPtr& SessionInterface = SessionManager->GetSessionInterface();
+
+	SessionInterface->ClearOnSessionFailureDelegates(this);
+	SessionInterface->DestroySession(PartySessionName, FOnDestroySessionCompleteDelegate::CreateSP(this, &FPartyPlatformSessionMonitor::HandleDestroySessionComplete));
 }
 
 void FPartyPlatformSessionMonitor::QueuePlatformSessionUpdate()
@@ -671,13 +681,13 @@ bool FPartyPlatformSessionMonitor::DoesLocalUserOwnPlatformSession()
 		}
 		else
 		{
-			for (UPartyMember* LocalMember : MonitoredParty->GetPartyMembers())
+			for (UPartyMember* Member : MonitoredParty->GetPartyMembers())
 			{
-				if (LocalMember->IsLocalPlayer())
+				if (Member->IsLocalPlayer())
 				{
 					for (const FPartyPlatformSessionInfo& SessionInfo : MonitoredParty->GetRepData().GetPlatformSessions())
 					{
-						if (SessionInfo.IsSessionOwner(*LocalMember))
+						if (SessionInfo.IsSessionOwner(*Member))
 						{
 							return true;
 						}
@@ -768,6 +778,8 @@ void FPartyPlatformSessionMonitor::HandleCreateSessionComplete(FName SessionName
 	{
 		SessionInitTracker.CompleteStep(Step_CreateSession);
 
+		SessionInterface->AddOnSessionFailureDelegate_Handle(FOnSessionFailureDelegate::CreateSP(this, &FPartyPlatformSessionMonitor::HandleSessionFailure));
+
 		if (MonitoredParty.IsValid())
 		{
 			MonitoredParty->SetIsMissingPlatformSession(false);
@@ -863,6 +875,7 @@ void FPartyPlatformSessionMonitor::HandleJoinSessionComplete(FName SessionName,
 	else if (bWasSuccessful && ensure(MonitoredParty.IsValid()))
 	{
 		MonitoredParty->SetIsMissingPlatformSession(false);
+		SessionInterface->AddOnSessionFailureDelegate_Handle(FOnSessionFailureDelegate::CreateSP(this, &FPartyPlatformSessionMonitor::HandleSessionFailure));
 
 		// For all local players
 		for (UPartyMember* PartyMember : MonitoredParty->GetPartyMembers())
@@ -1059,4 +1072,16 @@ bool FPartyPlatformSessionMonitor::HandleQueuedSessionUpdate(float)
 
 	// Only fire once - never retick
 	return false;
+}
+
+void FPartyPlatformSessionMonitor::HandleSessionFailure(const FUniqueNetId& LocalUserId, ESessionFailure::Type FailureType)
+{
+	UE_LOG(LogParty, Warning, TEXT("PartyPlatformSessionMonitor HandleSessionFailure LocalUserId=%s FailureType=%s"), *LocalUserId.ToDebugString(), LexToString(FailureType));
+
+	if (ensure(MonitoredParty.IsValid()))
+	{
+		MonitoredParty->SetIsMissingPlatformSession(true);
+
+		RetryTickerHandle = FTicker::GetCoreTicker().AddTicker(FTickerDelegate::CreateSP(this, &FPartyPlatformSessionMonitor::HandleRetryEstablishingSession), EstablishSessionRetryDelay);
+	}
 }
\ No newline at end of file
diff --git a/Engine/Plugins/Online/OnlineFramework/Source/Party/Private/Party/PartyPlatformSessionMonitor.h b/Engine/Plugins/Online/OnlineFramework/Source/Party/Private/Party/PartyPlatformSessionMonitor.h
--- a/Engine/Plugins/Online/OnlineFramework/Source/Party/Private/Party/PartyPlatformSessionMonitor.h
+++ b/Engine/Plugins/Online/OnlineFramework/Source/Party/Private/Party/PartyPlatformSessionMonitor.h
@@ -103,6 +103,7 @@ private:
 	void HandleJoinSessionComplete(FName SessionName, EOnJoinSessionCompleteResult::Type JoinSessionResult);
 	void HandleDestroySessionComplete(FName SessionName, bool bWasSuccessful);
 	bool HandleRetryEstablishingSession(float);
+	void HandleSessionFailure(const FUniqueNetId& LocalUserId, ESessionFailure::Type FailureType);
 
 	bool HandleQueuedSessionUpdate(float);
 
@@ -116,8 +117,6 @@ private:
 	/** Last session id we attempted to find to prevent repeated failures to find the same session */
 	TOptional<FSessionId> LastAttemptedFindSessionId;
 
-	bool bIsSessionMissing = false;
-
 	/** Do we have a console session update queued? */
 	bool bHasQueuedSessionUpdate = false;
 
diff --git a/Engine/Plugins/Online/OnlineFramework/Source/Party/Public/Party/SocialParty.h b/Engine/Plugins/Online/OnlineFramework/Source/Party/Public/Party/SocialParty.h
--- a/Engine/Plugins/Online/OnlineFramework/Source/Party/Public/Party/SocialParty.h
+++ b/Engine/Plugins/Online/OnlineFramework/Source/Party/Public/Party/SocialParty.h
@@ -205,6 +205,7 @@ PARTY_SCOPE:
 	bool IsCurrentlyLeaving() const;
 
 	void SetIsMissingPlatformSession(bool bInIsMissingPlatformSession);
+	bool IsMissingPlatformSession() { return bIsMissingPlatformSession; }
 
 	FPartyRepData& GetMutableRepData() { return *PartyDataReplicator; }
 
diff --git a/Engine/Plugins/Online/OnlineSubsystem/Source/Public/Interfaces/OnlineSessionInterface.h b/Engine/Plugins/Online/OnlineSubsystem/Source/Public/Interfaces/OnlineSessionInterface.h
--- a/Engine/Plugins/Online/OnlineSubsystem/Source/Public/Interfaces/OnlineSessionInterface.h
+++ b/Engine/Plugins/Online/OnlineSubsystem/Source/Public/Interfaces/OnlineSessionInterface.h
@@ -249,6 +249,18 @@ namespace ESessionFailure
 	};
 }
 
+/** Convert a ESessionFailure into a string */
+inline const TCHAR* LexToString(const ESessionFailure::Type Value)
+{
+	switch (Value)
+	{
+	case ESessionFailure::ServiceConnectionLost:
+		return TEXT("ServiceConnectionLost");
+	}
+
+	return TEXT("UnknownError");
+}
+
 /**
  * Delegate fired when an unexpected error occurs that impacts session connectivity or use
  *
