- Create missing platform sessions when becoming leader

diff --git a/Engine/Plugins/Online/OnlineFramework/Source/Party/Private/Party/SocialParty.cpp b/Engine/Plugins/Online/OnlineFramework/Source/Party/Private/Party/SocialParty.cpp
--- a/Engine/Plugins/Online/OnlineFramework/Source/Party/Private/Party/SocialParty.cpp
+++ b/Engine/Plugins/Online/OnlineFramework/Source/Party/Private/Party/SocialParty.cpp
@@ -584,29 +584,41 @@ void USocialParty::OnLocalPlayerIsLeaderChanged(bool bIsLeader)
 		// It's possible that membership changes resulting in this promotion also require updates to the session info
 		//	If we found out about the changes in membership before learning we're the leader, we were unable to update the rep data accordingly
 		//	So, upon becoming leader, we must do a sweep to account for any such changes we missed out on
-		TArray<FName> AllMemberPlatformSubsystems;
+		TArray<FName> SessionsToUpdate;
+		TArray<FName> SessionsToCreate;
 		for (UPartyMember* Member : GetPartyMembers())
 		{
 			FName PlatformOssName = Member->GetPlatformOssName();
 			if (!PlatformOssName.IsNone())
 			{
-				AllMemberPlatformSubsystems.AddUnique(PlatformOssName);
+				if (GetRepData().FindSessionInfo(PlatformOssName))
+				{
+					SessionsToUpdate.AddUnique(PlatformOssName);
+				}
+				else if (FPartyPlatformSessionManager::DoesOssNeedPartySession(PlatformOssName))
+				{
+					SessionsToCreate.AddUnique(PlatformOssName);
+				}
 			}
 		}
 		for (const FPartyPlatformSessionInfo& PlatformSessionInfo : GetRepData().GetPlatformSessions())
 		{
 			// Not using AddUnique so we can log if we are catching OSSs to remove
 			// TODO: Remove logging when we are sure this is fixed
-			if (!AllMemberPlatformSubsystems.Contains(PlatformSessionInfo.OssName))
+			if (!SessionsToUpdate.Contains(PlatformSessionInfo.OssName))
 			{
 				UE_LOG(LogParty, Verbose, TEXT("OnLocalPlayerIsLeaderChanged: Adding existing platform OSS [%s] which has no members"), *PlatformSessionInfo.OssName.ToString());
-				AllMemberPlatformSubsystems.Add(PlatformSessionInfo.OssName);
+				SessionsToUpdate.Add(PlatformSessionInfo.OssName);
 			}
 		}
-		for (FName PlatformOssName : AllMemberPlatformSubsystems)
+		for (FName PlatformOssName : SessionsToUpdate)
 		{
 			UpdatePlatformSessionLeader(PlatformOssName);
 		}
+		for (FName PlatformOssName : SessionsToCreate)
+		{
+			CreatePlatformSession(PlatformOssName);
+		}
 	}
 	else
 	{
@@ -1741,6 +1753,33 @@ void USocialParty::FinalizePartyLeave(EMemberExitedReason Reason)
 	PartyMembersById.Reset();
 }
 
+void USocialParty::CreatePlatformSession(FName PlatformOssName)
+{
+	if (!PlatformOssName.IsNone() && !GetRepData().FindSessionInfo(PlatformOssName))
+	{
+		FUniqueNetIdRepl OwnerPrimaryId;
+		for (UPartyMember* Member : GetPartyMembers())
+		{
+			if (Member->GetPlatformOssName() == PlatformOssName)
+			{
+				OwnerPrimaryId = Member->GetPrimaryNetId();
+				if (Member->IsLocalPlayer())
+				{
+					break; // Prefer Local players
+				}
+			}
+		}
+
+		if (OwnerPrimaryId.IsValid())
+		{
+			FPartyPlatformSessionInfo NewSessionInfo;
+			NewSessionInfo.OssName = PlatformOssName;
+			NewSessionInfo.OwnerPrimaryId = OwnerPrimaryId;
+			GetMutableRepData().UpdatePlatformSessionInfo(NewSessionInfo);
+		}
+	}
+}
+
 void USocialParty::UpdatePlatformSessionLeader(FName PlatformOssName)
 {
 	if (!IsLocalPlayerPartyLeader())
diff --git a/Engine/Plugins/Online/OnlineFramework/Source/Party/Public/Party/SocialParty.h b/Engine/Plugins/Online/OnlineFramework/Source/Party/Public/Party/SocialParty.h
--- a/Engine/Plugins/Online/OnlineFramework/Source/Party/Public/Party/SocialParty.h
+++ b/Engine/Plugins/Online/OnlineFramework/Source/Party/Public/Party/SocialParty.h
@@ -304,6 +304,7 @@ private:
 
 	void SetIsRequestingShutdown(bool bInRequestingShutdown);
 
+	void CreatePlatformSession(FName PlatformOssName);
 	void UpdatePlatformSessionLeader(FName PlatformOssName);
 
 	void HandlePreClientTravel(const FString& PendingURL, ETravelType TravelType, bool bIsSeamlessTravel);
