- Some compile error fixes that in UE5 target that didn't happen in others

diff --git a/Engine/Plugins/Online/OnlineFramework/Source/Hotfix/Private/OnlineHotfixManager.cpp b/Engine/Plugins/Online/OnlineFramework/Source/Hotfix/Private/OnlineHotfixManager.cpp
--- a/Engine/Plugins/Online/OnlineFramework/Source/Hotfix/Private/OnlineHotfixManager.cpp
+++ b/Engine/Plugins/Online/OnlineFramework/Source/Hotfix/Private/OnlineHotfixManager.cpp
@@ -826,13 +826,15 @@ FString UOnlineHotfixManager::GetConfigFileNamePath(const FString& IniName)
 FConfigFile* UOnlineHotfixManager::GetConfigFile(const FString& IniName)
 {
 	FString StrippedIniName(GetStrippedConfigFileName(IniName));
+	FString StrippedIniNameNoExtension = FPaths::GetBaseFilename(StrippedIniName);
+	
 	FConfigFile* ConfigFile = nullptr;
 	// Look for the first matching INI file entry
-	for (TMap<FString, FConfigFile>::TIterator It(*GConfig); It; ++It)
+	for (const FString& IniFilename : GConfig->GetFilenames())
 	{
-		if (It.Key().EndsWith(StrippedIniName))
+		if (IniFilename.EndsWith(StrippedIniName) || IniFilename == StrippedIniNameNoExtension)
 		{
-			ConfigFile = &It.Value();
+			ConfigFile = GConfig->FindConfigFile(IniFilename);
 			break;
 		}
 	}
diff --git a/Engine/Plugins/Runtime/AudioModulation/Source/AudioModulation/Private/AudioModulationProfileSerializer.h b/Engine/Plugins/Runtime/AudioModulation/Source/AudioModulation/Private/AudioModulationProfileSerializer.h
--- a/Engine/Plugins/Runtime/AudioModulation/Source/AudioModulation/Private/AudioModulationProfileSerializer.h
+++ b/Engine/Plugins/Runtime/AudioModulation/Source/AudioModulation/Private/AudioModulationProfileSerializer.h
@@ -107,7 +107,8 @@ namespace AudioModulation
 				return false;
 			}
 
-			FConfigFile& ConfigFile = GConfig->FindOrAdd(ConfigPath);
+			// this will create and try to read in if not already loaded
+			FConfigFile& ConfigFile = *GConfig->Find(ConfigPath);
 			ConfigFile.Reset();
 
 			for (const FSoundControlBusMixStage& Stage : InBusMix.MixStages)
diff --git a/Engine/Source/Runtime/Core/Private/Misc/ConfigCacheIni.cpp b/Engine/Source/Runtime/Core/Private/Misc/ConfigCacheIni.cpp
--- a/Engine/Source/Runtime/Core/Private/Misc/ConfigCacheIni.cpp
+++ b/Engine/Source/Runtime/Core/Private/Misc/ConfigCacheIni.cpp
@@ -2650,11 +2650,7 @@ bool FConfigCacheIni::EmptySectionsMatchingString( const TCHAR* SectionString, c
  */
 void FConfigCacheIni::GetConfigFilenames(TArray<FString>& ConfigFilenames)
 {
-	// copy from our map to the array
-	for (FConfigCacheIni::TIterator It(*this); It; ++It)
-	{
-		ConfigFilenames.Add(*(It.Key()));
-	}
+	ConfigFilenames = GetFilenames();
 }
 
 /**
