Workaround for internal compiler error on MSVC

diff --git a/Engine/Plugins/Online/OnlineFramework/Source/Party/Private/Party/PartyPlatformSessionMonitor.cpp b/Engine/Plugins/Online/OnlineFramework/Source/Party/Private/Party/PartyPlatformSessionMonitor.cpp
--- a/Engine/Plugins/Online/OnlineFramework/Source/Party/Private/Party/PartyPlatformSessionMonitor.cpp
+++ b/Engine/Plugins/Online/OnlineFramework/Source/Party/Private/Party/PartyPlatformSessionMonitor.cpp
@@ -166,7 +166,11 @@ bool FPartyPlatformSessionManager::FindSession(const USocialUser& User, const FO
 	FUniqueNetIdRepl UserPlatformId;
 	if (const FOnlineUserPresence* PlatformPresence = User.GetFriendPresenceInfo(ESocialSubsystem::Platform))
 	{
-		SessionId = PlatformPresence->SessionId ? PlatformPresence->SessionId->ToString() : FSessionId();
+		if (PlatformPresence->SessionId)
+		{
+			SessionId = PlatformPresence->SessionId->ToString();
+		}
+
 		UserPlatformId = User.GetUserId(ESocialSubsystem::Platform);
 	}
 	return FindSessionInternal(SessionId, UserPlatformId, OnAttemptComplete);
@@ -1168,4 +1172,4 @@ void FPartyPlatformSessionMonitor::HandleSessionFailure(const FUniqueNetId& Loca
 
 		RetryTickerHandle = FTicker::GetCoreTicker().AddTicker(FTickerDelegate::CreateSP(this, &FPartyPlatformSessionMonitor::HandleRetryEstablishingSession), EstablishSessionRetryDelay);
 	}
-}
\ No newline at end of file
+}
