Deprecated legacy online name aliases

diff --git a/Engine/Plugins/Online/OnlineFramework/Source/Party/Private/Party/PartyPlatformSessionMonitor.cpp b/Engine/Plugins/Online/OnlineFramework/Source/Party/Private/Party/PartyPlatformSessionMonitor.cpp
--- a/Engine/Plugins/Online/OnlineFramework/Source/Party/Private/Party/PartyPlatformSessionMonitor.cpp
+++ b/Engine/Plugins/Online/OnlineFramework/Source/Party/Private/Party/PartyPlatformSessionMonitor.cpp
@@ -419,7 +419,7 @@ const FOnlinePartyTypeId& FPartyPlatformSessionMonitor::GetMonitoredPartyTypeId(
 
 EOnlineSessionState::Type FPartyPlatformSessionMonitor::GetOssSessionState() const
 {
-	FNamedOnlineSession* PlatformSession = SessionManager->GetSessionInterface()->GetNamedSession(PartySessionName);
+	FNamedOnlineSession* PlatformSession = SessionManager->GetSessionInterface()->GetNamedSession(NAME_PartySession);
 	return PlatformSession ? PlatformSession->SessionState : EOnlineSessionState::NoSession;
 }
 
@@ -432,7 +432,7 @@ void FPartyPlatformSessionMonitor::EvaluateCurrentSession()
 	}
 
 	const IOnlineSessionPtr& SessionInterface = SessionManager->GetSessionInterface();
-	if (const FNamedOnlineSession* Session = SessionInterface->GetNamedSession(PartySessionName))
+	if (const FNamedOnlineSession* Session = SessionInterface->GetNamedSession(NAME_PartySession))
 	{
 		// TODO: We need to check which of the local players are in the session and process them accordingly 
 		// We already have a platform session session, so we should be all set. Just check the session ids to make sure we have the correct session.
@@ -549,11 +549,11 @@ void FPartyPlatformSessionMonitor::CreateSession(const FUniqueNetIdRepl& LocalUs
 						if (ForcePlatformSessionCreationFailure != 0)
 						{
 							UE_LOG(LogParty, Warning, TEXT("Forcing session creation failure"));
-							HandleCreateSessionComplete(PartySessionName, false);
+							HandleCreateSessionComplete(NAME_PartySession, false);
 						}
 						else
 						{
-							SessionManager->GetSessionInterface()->CreateSession(*LocalUserPlatformId, PartySessionName, SessionSettings);
+							SessionManager->GetSessionInterface()->CreateSession(*LocalUserPlatformId, NAME_PartySession, SessionSettings);
 						}
 					}
 					return false; // Don't retick
@@ -561,7 +561,7 @@ void FPartyPlatformSessionMonitor::CreateSession(const FUniqueNetIdRepl& LocalUs
 			return;
 		}
 #endif
-		SessionInterface->CreateSession(*LocalUserPlatformId, PartySessionName, SessionSettings);
+		SessionInterface->CreateSession(*LocalUserPlatformId, NAME_PartySession, SessionSettings);
 		UE_LOG(LogParty, Verbose, TEXT("PartyPlatformSessionMonitor creating session with the following parameters: "));
 		DumpSessionSettings(&SessionSettings);
 	}
@@ -573,7 +573,7 @@ void FPartyPlatformSessionMonitor::AddLocalPlayerToSession(UPartyMember* PartyMe
 		!MonitoredParty->IsMissingPlatformSession())
 	{
 		const IOnlineSessionPtr& SessionInterface = SessionManager->GetSessionInterface();
-		const FNamedOnlineSession* Session = SessionInterface->GetNamedSession(PartySessionName);
+		const FNamedOnlineSession* Session = SessionInterface->GetNamedSession(NAME_PartySession);
 		if (ensure(Session && Session->SessionInfo.IsValid()))
 		{
 			const FSessionId SessionId = Session->GetSessionIdStr();
@@ -603,7 +603,7 @@ void FPartyPlatformSessionMonitor::RemoveLocalPlayerFromSession(UPartyMember* Pa
 		!MonitoredParty->IsMissingPlatformSession())
 	{
 		const IOnlineSessionPtr& SessionInterface = SessionManager->GetSessionInterface();
-		const FNamedOnlineSession* Session = SessionInterface->GetNamedSession(PartySessionName);
+		const FNamedOnlineSession* Session = SessionInterface->GetNamedSession(NAME_PartySession);
 
 		if (ensure(Session))
 		{
@@ -672,11 +672,11 @@ void FPartyPlatformSessionMonitor::JoinSession(const FOnlineSessionSearchResult&
 					if (ForcePlatformSessionCreationFailure != 0)
 					{
 						UE_LOG(LogParty, Warning, TEXT("Forcing session join failure"));
-						HandleJoinSessionComplete(PartySessionName, EOnJoinSessionCompleteResult::UnknownError);
+						HandleJoinSessionComplete(NAME_PartySession, EOnJoinSessionCompleteResult::UnknownError);
 					}
 					else
 					{
-						SessionManager->GetSessionInterface()->JoinSession(*LocalUserPlatformId, PartySessionName, SearchResultCopy);
+						SessionManager->GetSessionInterface()->JoinSession(*LocalUserPlatformId, NAME_PartySession, SearchResultCopy);
 					}
 				}
 				return false; // Don't retick
@@ -685,7 +685,7 @@ void FPartyPlatformSessionMonitor::JoinSession(const FOnlineSessionSearchResult&
 	}
 #endif
 		
-	if (!SessionInterface->JoinSession(*LocalUserPlatformId, PartySessionName, SearchResultCopy))
+	if (!SessionInterface->JoinSession(*LocalUserPlatformId, NAME_PartySession, SearchResultCopy))
 	{
 		UE_LOG(LogParty, Warning, TEXT("JoinSession call failed for session [%s]."), *SessionSearchResult.GetSessionIdStr());
 		TargetSessionId = FSessionId();
@@ -705,7 +705,7 @@ void FPartyPlatformSessionMonitor::LeaveSession()
 	const IOnlineSessionPtr& SessionInterface = SessionManager->GetSessionInterface();
 
 	SessionInterface->ClearOnSessionFailureDelegates(this);
-	SessionInterface->DestroySession(PartySessionName, FOnDestroySessionCompleteDelegate::CreateSP(this, &FPartyPlatformSessionMonitor::HandleDestroySessionComplete));
+	SessionInterface->DestroySession(NAME_PartySession, FOnDestroySessionCompleteDelegate::CreateSP(this, &FPartyPlatformSessionMonitor::HandleDestroySessionComplete));
 }
 
 void FPartyPlatformSessionMonitor::QueuePlatformSessionUpdate()
@@ -815,7 +815,7 @@ void FPartyPlatformSessionMonitor::HandlePartyMemberInitialized(UPartyMember* In
 {
 	if (IsTencentPlatform() && InitializedMember->GetPlatformOssName() == TENCENT_SUBSYSTEM)
 	{
-		SessionManager->GetSessionInterface()->RegisterPlayer(PartySessionName, *InitializedMember->GetRepData().GetPlatformUniqueId(), false);
+		SessionManager->GetSessionInterface()->RegisterPlayer(NAME_PartySession, *InitializedMember->GetRepData().GetPlatformUniqueId(), false);
 	}
 
 	// If a local player joined the party (split screen) we add them to the platform session
@@ -832,7 +832,7 @@ void FPartyPlatformSessionMonitor::HandlePartyMemberLeft(UPartyMember* OldMember
 
 	if (IsTencentPlatform() && OldMember->GetPlatformOssName() == TENCENT_SUBSYSTEM)
 	{
-		SessionManager->GetSessionInterface()->UnregisterPlayer(PartySessionName, *OldMember->GetRepData().GetPlatformUniqueId());
+		SessionManager->GetSessionInterface()->UnregisterPlayer(NAME_PartySession, *OldMember->GetRepData().GetPlatformUniqueId());
 	}
 
 	if (OldMember->IsLocalPlayer())
@@ -975,7 +975,7 @@ void FPartyPlatformSessionMonitor::HandleJoinSessionComplete(FName SessionName,
 					MemberIdsOnPlatform.Add(PartyMember->GetRepData().GetPlatformUniqueId().GetUniqueNetId().ToSharedRef());
 				}
 			}
-			SessionInterface->RegisterPlayers(PartySessionName, MemberIdsOnPlatform);
+			SessionInterface->RegisterPlayers(NAME_PartySession, MemberIdsOnPlatform);
 		}
 	}
 	else
@@ -1137,14 +1137,14 @@ bool FPartyPlatformSessionMonitor::HandleQueuedSessionUpdate(float)
 		const IOnlineSessionPtr& SessionInterface = SessionManager->GetSessionInterface();
 
 		// Make sure the party session is in a fully created state and is not destroying
-		FNamedOnlineSession* PlatformSession = SessionInterface->GetNamedSession(PartySessionName);
+		FNamedOnlineSession* PlatformSession = SessionInterface->GetNamedSession(NAME_PartySession);
 		if (PlatformSession)
 		{
 			if (PlatformSession->SessionState >= EOnlineSessionState::Pending && PlatformSession->SessionState <= EOnlineSessionState::Ended)
 			{
 				if (ConfigurePlatformSessionSettings(PlatformSession->SessionSettings))
 				{
-					if (!SessionInterface->UpdateSession(PartySessionName, PlatformSession->SessionSettings, true))
+					if (!SessionInterface->UpdateSession(NAME_PartySession, PlatformSession->SessionSettings, true))
 					{
 						UE_LOG(LogParty, Warning, TEXT("PartyPlatformSessionMonitor call to UpdateSession failed"));
 					}
diff --git a/Engine/Plugins/Online/OnlineSubsystemOculus/Source/Private/OculusUpdateSessionCallbackProxy.cpp b/Engine/Plugins/Online/OnlineSubsystemOculus/Source/Private/OculusUpdateSessionCallbackProxy.cpp
--- a/Engine/Plugins/Online/OnlineSubsystemOculus/Source/Private/OculusUpdateSessionCallbackProxy.cpp
+++ b/Engine/Plugins/Online/OnlineSubsystemOculus/Source/Private/OculusUpdateSessionCallbackProxy.cpp
@@ -31,7 +31,7 @@ void UOculusUpdateSessionCallbackProxy::Activate()
 
 		FOnlineSessionSettings Settings;
 		Settings.bShouldAdvertise = bShouldEnqueueInMatchmakingPool;
-		OculusSessionInterface->UpdateSession(GameSessionName, Settings);
+		OculusSessionInterface->UpdateSession(NAME_GameSession, Settings);
 	}
 	else
 	{
diff --git a/Engine/Source/Runtime/CoreOnline/Public/Online/CoreOnline.h b/Engine/Source/Runtime/CoreOnline/Public/Online/CoreOnline.h
--- a/Engine/Source/Runtime/CoreOnline/Public/Online/CoreOnline.h
+++ b/Engine/Source/Runtime/CoreOnline/Public/Online/CoreOnline.h
@@ -7,10 +7,14 @@
 #include "Online/CoreOnlineFwd.h"
 #include "UObject/ObjectMacros.h"
 
-#define GameSessionName NAME_GameSession
-#define PartySessionName NAME_PartySession
-#define GamePort NAME_GamePort
-#define BeaconPort NAME_BeaconPort
+UE_DEPRECATED(5.0, "Use NAME_GameSession.")
+inline constexpr EName GameSessionName = NAME_GameSession;
+UE_DEPRECATED(5.0, "Use NAME_PartySession.")
+inline constexpr EName PartySessionName = NAME_PartySession;
+UE_DEPRECATED(5.0, "Use NAME_GamePort.")
+inline constexpr EName GamePort = NAME_GamePort;
+UE_DEPRECATED(5.0, "Use NAME_BeaconPort.")
+inline constexpr EName BeaconPort = NAME_BeaconPort;
 
 USTRUCT(noexport)
 struct FJoinabilitySettings
