Fix hotfixing in UE5.

diff --git a/Engine/Plugins/Online/OnlineFramework/Source/Hotfix/Private/OnlineHotfixManager.cpp b/Engine/Plugins/Online/OnlineFramework/Source/Hotfix/Private/OnlineHotfixManager.cpp
--- a/Engine/Plugins/Online/OnlineFramework/Source/Hotfix/Private/OnlineHotfixManager.cpp
+++ b/Engine/Plugins/Online/OnlineFramework/Source/Hotfix/Private/OnlineHotfixManager.cpp
@@ -823,33 +823,56 @@ FString UOnlineHotfixManager::GetStrippedConfigFileName(const FString& IniName)
 	return StrippedIniName;
 }
 
-FString UOnlineHotfixManager::GetConfigFileNamePath(const FString& IniName)
+FString UOnlineHotfixManager::BuildConfigCacheKey(const FString& IniName)
 {
-	return FPaths::GeneratedConfigDir() + ANSI_TO_TCHAR(FPlatformProperties::PlatformName()) / IniName;
+	const FString IniNameNoExtension = FPaths::GetBaseFilename(IniName);
+
+	// Known ini files such as Engine, Game, etc.. are referred to as just the name with no extension within the config system.
+	if (GConfig->IsKnownConfigName(FName(*IniNameNoExtension, FNAME_Find)))
+	{
+		return IniNameNoExtension;
+	}
+	else
+	{
+		// Non-known ini files are looked up using their full path.
+		return GConfig->GetDestIniFilename(*IniNameNoExtension, nullptr, *FPaths::GeneratedConfigDir());
+	}
 }
 
 FConfigFile* UOnlineHotfixManager::GetConfigFile(const FString& IniName)
 {
-	FString StrippedIniName(GetStrippedConfigFileName(IniName));
-	FString StrippedIniNameNoExtension = FPaths::GetBaseFilename(StrippedIniName);
+	const FString StrippedIniName(GetStrippedConfigFileName(IniName));
+	const FString StrippedIniNameNoExtension = FPaths::GetBaseFilename(StrippedIniName);
 	
 	FConfigFile* ConfigFile = nullptr;
-	// Look for the first matching INI file entry
-	for (const FString& IniFilename : GConfig->GetFilenames())
+
+	// Start by searching for a known config name.
+	if (GConfig->IsKnownConfigName(FName(*StrippedIniNameNoExtension, FNAME_Find)))
 	{
-		if (IniFilename.EndsWith(StrippedIniName) || IniFilename == StrippedIniNameNoExtension)
+		ConfigFile = GConfig->FindConfigFile(StrippedIniNameNoExtension);
+	}
+
+	// Fall back to a partial path search.
+	if (ConfigFile == nullptr)
+	{
+		// Look for the first matching INI file entry
+		for (const FString& IniFilename : GConfig->GetFilenames())
 		{
-			ConfigFile = GConfig->FindConfigFile(IniFilename);
-			break;
+			if (IniFilename.EndsWith(StrippedIniName) || IniFilename == StrippedIniNameNoExtension)
+			{
+				ConfigFile = GConfig->FindConfigFile(IniFilename);
+				break;
+			}
 		}
 	}
-	// If not found, add this file to the config cache
+
+	// If not found, add this file to the config cache.
 	if (ConfigFile == nullptr)
 	{
-		const FString IniNameWithPath = GetConfigFileNamePath(StrippedIniName);
+		const FString ProcessedName(BuildConfigCacheKey(StrippedIniName));
 		FConfigFile Empty;
-		GConfig->SetFile(IniNameWithPath, &Empty);
-		ConfigFile = GConfig->Find(IniNameWithPath);
+		GConfig->SetFile(ProcessedName, &Empty);
+		ConfigFile = GConfig->Find(ProcessedName);
 	}
 	check(ConfigFile);
 	// We never want to save these merged files
@@ -1295,7 +1318,7 @@ void UOnlineHotfixManager::OnReadFileProgress(const FString& FileName, uint64 By
 
 UOnlineHotfixManager::FConfigFileBackup& UOnlineHotfixManager::BackupIniFile(const FString& IniName, const FConfigFile* ConfigFile)
 {
-	FString BackupIniName = GetConfigFileNamePath(GetStrippedConfigFileName(IniName));
+	FString BackupIniName = BuildConfigCacheKey(GetStrippedConfigFileName(IniName));
 	if (FConfigFileBackup* Backup = IniBackups.FindByPredicate([&BackupIniName](const FConfigFileBackup& Entry) { return Entry.IniName == BackupIniName; }))
 	{
 		// Only store one copy of each ini file, consisting of the original state
@@ -1325,7 +1348,7 @@ void UOnlineHotfixManager::RestoreBackupIniFiles()
 	{
 		if (FileHeader.FileName.EndsWith(TEXT(".INI")))
 		{
-			const FString ProcessedName = GetConfigFileNamePath(GetStrippedConfigFileName(FileHeader.FileName));
+			const FString ProcessedName = BuildConfigCacheKey(GetStrippedConfigFileName(FileHeader.FileName));
 			for (int32 Index = 0; Index < IniBackups.Num(); Index++)
 			{
 				const FConfigFileBackup& BackupFile = IniBackups[Index];
@@ -1346,7 +1369,7 @@ void UOnlineHotfixManager::RestoreBackupIniFiles()
 	{
 		if (FileHeader.FileName.EndsWith(TEXT(".INI")))
 		{
-			const FString ProcessedName = GetConfigFileNamePath(GetStrippedConfigFileName(FileHeader.FileName));
+			const FString ProcessedName = BuildConfigCacheKey(GetStrippedConfigFileName(FileHeader.FileName));
 			for (int32 Index = 0; Index < IniBackups.Num(); Index++)
 			{
 				const FConfigFileBackup& BackupFile = IniBackups[Index];
diff --git a/Engine/Plugins/Online/OnlineFramework/Source/Hotfix/Public/OnlineHotfixManager.h b/Engine/Plugins/Online/OnlineFramework/Source/Hotfix/Public/OnlineHotfixManager.h
--- a/Engine/Plugins/Online/OnlineFramework/Source/Hotfix/Public/OnlineHotfixManager.h
+++ b/Engine/Plugins/Online/OnlineFramework/Source/Hotfix/Public/OnlineHotfixManager.h
@@ -202,8 +202,8 @@ protected:
 	/** @return the config file entry for the ini file name in question */
 	FConfigFile* GetConfigFile(const FString& IniName);
 
-	/** @return the config file name with full path info */
-	FString GetConfigFileNamePath(const FString& IniName);
+	/** @return the config cache key used to associate ini file entries within the config cache */
+	FString BuildConfigCacheKey(const FString& IniName);
 	/** @return the config file name after stripping any extra info (platform, debug prefix, etc.) */
 	virtual FString GetStrippedConfigFileName(const FString& IniName);
 
diff --git a/Engine/Source/Runtime/Core/Private/Misc/ConfigCacheIni.cpp b/Engine/Source/Runtime/Core/Private/Misc/ConfigCacheIni.cpp
--- a/Engine/Source/Runtime/Core/Private/Misc/ConfigCacheIni.cpp
+++ b/Engine/Source/Runtime/Core/Private/Misc/ConfigCacheIni.cpp
@@ -2487,7 +2487,14 @@ void FConfigCacheIni::LoadFile( const FString& Filename, const FConfigFile* Fall
 
 void FConfigCacheIni::SetFile( const FString& Filename, const FConfigFile* NewConfigFile )
 {
-	Add(Filename, *NewConfigFile);
+	if (FConfigFile* FoundFile = KnownFiles.GetMutableFile(FName(*Filename, FNAME_Find)))
+	{
+		*FoundFile = *NewConfigFile;
+	}
+	else
+	{
+		Add(Filename, *NewConfigFile);
+	}
 }
 
 
