Fix platform monitor to make sure party is valid before firing analytic event

diff --git a/Engine/Plugins/Online/OnlineFramework/Source/Party/Private/Party/PartyPlatformSessionMonitor.cpp b/Engine/Plugins/Online/OnlineFramework/Source/Party/Private/Party/PartyPlatformSessionMonitor.cpp
--- a/Engine/Plugins/Online/OnlineFramework/Source/Party/Private/Party/PartyPlatformSessionMonitor.cpp
+++ b/Engine/Plugins/Online/OnlineFramework/Source/Party/Private/Party/PartyPlatformSessionMonitor.cpp
@@ -992,7 +992,10 @@ void FPartyPlatformSessionMonitor::HandleJoinSessionComplete(FName SessionName,
 		ProcessJoinFailure();
 	}
 
-	MonitoredParty->JoinSessionCompleteAnalytics(*TargetSessionId, LexToString(JoinSessionResult));
+	if (MonitoredParty.IsValid())
+	{
+		MonitoredParty->JoinSessionCompleteAnalytics(*TargetSessionId, LexToString(JoinSessionResult));
+	}
 }
 
 void FPartyPlatformSessionMonitor::HandleDestroySessionComplete(FName SessionName, bool bWasSuccessful)
