back out 16307904: Prevent async loading during hotfixing / content patching.

#rnx

diff --git a/Engine/Plugins/Online/OnlineFramework/Source/Hotfix/Private/OnlineHotfixManager.cpp b/Engine/Plugins/Online/OnlineFramework/Source/Hotfix/Private/OnlineHotfixManager.cpp
--- a/Engine/Plugins/Online/OnlineFramework/Source/Hotfix/Private/OnlineHotfixManager.cpp
+++ b/Engine/Plugins/Online/OnlineFramework/Source/Hotfix/Private/OnlineHotfixManager.cpp
@@ -23,8 +23,6 @@
 #include "Curves/CurveFloat.h"
 #include "Engine/BlueprintGeneratedClass.h"
 
-#include "Serialization/AsyncLoadingLock.h"
-
 #ifdef WITH_ONLINETRACING
 #include "OnlineTracingModule.h"
 #endif
@@ -179,15 +177,6 @@ UOnlineHotfixManager::UOnlineHotfixManager() :
 	GameContentPath = FString() / FApp::GetProjectName() / TEXT("Content");
 }
 
-UOnlineHotfixManager::UOnlineHotfixManager(FVTableHelper& Helper)
-	: Super(Helper)
-{
-}
-
-UOnlineHotfixManager::~UOnlineHotfixManager()
-{
-}
-
 UOnlineHotfixManager* UOnlineHotfixManager::Get(UWorld* World)
 {
 	UOnlineHotfixManager* DefaultObject = UOnlineHotfixManager::StaticClass()->GetDefaultObject<UOnlineHotfixManager>();
@@ -269,7 +258,6 @@ void UOnlineHotfixManager::Cleanup()
 	}
 	OnlineTitleFile = nullptr;
 	bHotfixingInProgress = false;
-	AsyncLock = nullptr;
 }
 
 void UOnlineHotfixManager::StartHotfixProcess()
@@ -431,34 +419,14 @@ void UOnlineHotfixManager::OnEnumerateFilesComplete(bool bWasSuccessful, const F
 		{
 			if (RemovedHotfixFileList.Num() > 0)
 			{
-				UE_LOG(LogHotfixManager, Display, TEXT("Files have been removed since last check. Reverting."));
-
-				// Prevent async loading while reverting hotfixes.
-				check(!AsyncLock);
-				AsyncLock = MakeUnique<FAsyncLoadingLock>(TEXT("RevertHotfix"));
-				AsyncLock->Acquire(
-					FOnAsyncLoadingLockAcquired::CreateWeakLambda(
-						this,
-						[this]()
-						{
-							{
-								// No changes, just reverts
-								// Perform any undo operations needed
-								RestoreBackupIniFiles();
-								UnmountHotfixFiles();
-							}
-
-							AsyncLock->Release();
-
-							// Trigger completion outside of async lock.
-							TriggerHotfixComplete(EHotfixResult::SuccessNoChange);
-						}));
-			}
-			else
-			{
-				UE_LOG(LogHotfixManager, Display, TEXT("Returned hotfix data is the same as last application, skipping the apply phase"));
-				TriggerHotfixComplete(EHotfixResult::SuccessNoChange);
+				// No changes, just reverts
+				// Perform any undo operations needed
+				RestoreBackupIniFiles();
+				UnmountHotfixFiles();
 			}
+
+			UE_LOG(LogHotfixManager, Display, TEXT("Returned hotfix data is the same as last application, skipping the apply phase"));
+			TriggerHotfixComplete(EHotfixResult::SuccessNoChange);
 		}
 	}
 	else
@@ -680,25 +648,8 @@ void UOnlineHotfixManager::OnReadFileComplete(bool bWasSuccessful, const FString
 			PendingHotfixFiles.Remove(FileName);
 			if (PendingHotfixFiles.Num() == 0)
 			{
-				// Prevent async loading while applying hotfixes.
-				check(!AsyncLock);
-				AsyncLock = MakeUnique<FAsyncLoadingLock>(TEXT("ApplyHotfix"));
-				AsyncLock->Acquire(
-					FOnAsyncLoadingLockAcquired::CreateWeakLambda(
-						this,
-						[this]()
-						{
-							EHotfixResult Result = EHotfixResult::Success;
-
-							{
-								Result = ApplyHotfix();
-							}
-
-							AsyncLock->Release();
-
-							// Trigger completion outside of async lock.
-							TriggerHotfixComplete(Result);
-						}));
+				// All files have been downloaded so now apply the files
+				ApplyHotfix();
 			}
 		}
 		else
@@ -717,7 +668,7 @@ void UOnlineHotfixManager::UpdateProgress(uint32 FileCount, uint64 UpdateSize)
 	TriggerOnHotfixProgressDelegates(NumDownloaded, TotalFiles, NumBytes, TotalBytes);
 }
 
-EHotfixResult UOnlineHotfixManager::ApplyHotfix()
+void UOnlineHotfixManager::ApplyHotfix()
 {
 	// Perform any undo operations needed
 	// This occurs same frame as the application of new hotfixes
@@ -729,7 +680,8 @@ EHotfixResult UOnlineHotfixManager::ApplyHotfix()
 		if (!ApplyHotfixProcessing(FileHeader))
 		{
 			UE_LOG(LogHotfixManager, Error, TEXT("Couldn't apply hotfix file (%s)"), *FileHeader.FileName);
-			return EHotfixResult::Failed;
+			TriggerHotfixComplete(EHotfixResult::Failed);
+			return;
 		}
 		// Let anyone listening know we just processed this file
 		TriggerOnHotfixProcessedFileDelegates(FileHeader.FileName, GetCachedDirectory() / FileHeader.DLName);
@@ -746,16 +698,11 @@ EHotfixResult UOnlineHotfixManager::ApplyHotfix()
 		UE_LOG(LogHotfixManager, Display, TEXT("Hotfix has detected PAK files containing currently loaded maps, so a level load is needed"));
 		Result = EHotfixResult::SuccessNeedsReload;
 	}
-
-	return Result;
+	TriggerHotfixComplete(Result);
 }
 
 void UOnlineHotfixManager::TriggerHotfixComplete(EHotfixResult HotfixResult)
 {
-	// Flush the async loader to prevent updating data while it is accessed for loading.
-	// This may cause a hitch in some failure cases.
-	FlushAsyncLoading();
-
 	if (HotfixResult != EHotfixResult::Failed && HotfixResult != EHotfixResult::SuccessNoChange)
 	{
 		PatchAssetsFromIniFiles();
diff --git a/Engine/Plugins/Online/OnlineFramework/Source/Hotfix/Public/OnlineHotfixManager.h b/Engine/Plugins/Online/OnlineFramework/Source/Hotfix/Public/OnlineHotfixManager.h
--- a/Engine/Plugins/Online/OnlineFramework/Source/Hotfix/Public/OnlineHotfixManager.h
+++ b/Engine/Plugins/Online/OnlineFramework/Source/Hotfix/Public/OnlineHotfixManager.h
@@ -15,8 +15,6 @@
 
 HOTFIX_API DECLARE_LOG_CATEGORY_EXTERN(LogHotfixManager, Display, All);
 
-class FAsyncLoadingLock;
-
 UENUM()
 enum class EHotfixResult : uint8
 {
@@ -159,8 +157,6 @@ protected:
 	uint64 NumBytes;
 	/** Some title file interfaces aren't re-entrant so handle it ourselves */
 	bool bHotfixingInProgress;
-	/** Lock to prevent async loading during the hotfixing process. */
-	TUniquePtr<FAsyncLoadingLock> AsyncLock;
 	/** Set to true if any PAK file contains an update to a level that is currently loaded */
 	bool bHotfixNeedsMapReload;
 #if !UE_BUILD_SHIPPING
@@ -178,7 +174,7 @@ protected:
 	virtual void Init();
 	virtual void Cleanup();
 	/** Looks at each file returned via the hotfix and processes them */
-	EHotfixResult ApplyHotfix();
+	void ApplyHotfix();
 	/** Cleans up and fires the delegate indicating it's done */
 	void TriggerHotfixComplete(EHotfixResult HotfixResult);
 	/** Checks each file listed to see if it is a hotfix file to process */
@@ -315,8 +311,6 @@ protected:
 
 public:
 	UOnlineHotfixManager();
-	UOnlineHotfixManager(FVTableHelper& Helper);
-	virtual ~UOnlineHotfixManager();
 
 	/** Tells the hotfix manager which OSS to use. Uses the default if empty */
 	UPROPERTY(Config)
diff --git a/Engine/Source/Runtime/CoreUObject/Private/Serialization/AsyncLoadingLock.cpp b/Engine/Source/Runtime/CoreUObject/Private/Serialization/AsyncLoadingLock.cpp
deleted file mode 100644
diff --git a/Engine/Source/Runtime/CoreUObject/Public/Serialization/AsyncLoadingLock.h b/Engine/Source/Runtime/CoreUObject/Public/Serialization/AsyncLoadingLock.h
deleted file mode 100644
