Delay party join attempts until internal persistent party creation has completed.

diff --git a/Engine/Plugins/Online/OnlineFramework/Source/Party/Private/SocialManager.cpp b/Engine/Plugins/Online/OnlineFramework/Source/Party/Private/SocialManager.cpp
--- a/Engine/Plugins/Online/OnlineFramework/Source/Party/Private/SocialManager.cpp
+++ b/Engine/Plugins/Online/OnlineFramework/Source/Party/Private/SocialManager.cpp
@@ -73,6 +73,7 @@ const FName USocialManager::FJoinPartyAttempt::Step_QueryJoinability = TEXT("Que
 const FName USocialManager::FJoinPartyAttempt::Step_LeaveCurrentParty = TEXT("LeaveCurrentParty");
 const FName USocialManager::FJoinPartyAttempt::Step_JoinParty = TEXT("JoinParty");
 const FName USocialManager::FJoinPartyAttempt::Step_DeferredPartyCreation = TEXT("DeferredPartyCreation");
+const FName USocialManager::FJoinPartyAttempt::Step_WaitForPersistentPartyCreation = TEXT("WaitForPersistentPartyCreation");
 
 //////////////////////////////////////////////////////////////////////////
 // USocialManager
@@ -409,6 +410,10 @@ void USocialManager::CreatePersistentParty(const FOnCreatePartyAttemptComplete&
 {
 	UE_LOG(LogParty, Log, TEXT("Attempting to create new persistent party"));
 
+	// If the player tries to join a party while persistent party creation is in progress, the persistent party
+	// should be left when creation completes before attempting to handle the join attempt.
+	bCreatingPersistentParty = true;
+
 	// The persistent party starts off closed by default, and will update its config as desired after initializing
 	FPartyConfiguration InitialPersistentPartyConfig;
 	InitialPersistentPartyConfig.JoinRequestAction = EJoinRequestAction::Manual;
@@ -418,7 +423,48 @@ void USocialManager::CreatePersistentParty(const FOnCreatePartyAttemptComplete&
 	InitialPersistentPartyConfig.InvitePermissions = PartySystemPermissions::EPermissionType::Noone;
 	InitialPersistentPartyConfig.MaxMembers = USocialSettings::GetDefaultMaxPartySize();
 
-	CreateParty(IOnlinePartySystem::GetPrimaryPartyTypeId(), InitialPersistentPartyConfig, OnCreatePartyComplete);
+	CreateParty(IOnlinePartySystem::GetPrimaryPartyTypeId(), InitialPersistentPartyConfig, FOnCreatePartyAttemptComplete::CreateUObject(this, &USocialManager::OnCreatePersistentPartyCompleteInternal, OnCreatePartyComplete));
+}
+
+void USocialManager::OnCreatePersistentPartyCompleteInternal(ECreatePartyCompletionResult Result, FOnCreatePartyAttemptComplete OnCreatePartyComplete)
+{
+	ABORT_DURING_SHUTDOWN();
+
+	bCreatingPersistentParty = false;
+
+	// Notify completion before possibly tearing down the party for a join attempt.
+	OnCreatePartyComplete.ExecuteIfBound(Result);
+
+	// Leave persistent party if a join attempt was made to a different persistent party while creation was underway.
+	// OnPartyLeft will handle joining the other persistent party.
+	if (FJoinPartyAttempt* JoinAttempt = JoinAttemptsByTypeId.Find(IOnlinePartySystem::GetPrimaryPartyTypeId()))
+	{
+		JoinAttempt->ActionTimeTracker.CompleteStep(FJoinPartyAttempt::Step_WaitForPersistentPartyCreation);
+
+		if (Result == ECreatePartyCompletionResult::Succeeded)
+		{
+			if (USocialParty* ExistingParty = GetPartyInternal(IOnlinePartySystem::GetPrimaryPartyTypeId(), true))
+			{
+				// Another party of the same type exists - leave that one first
+				JoinAttempt->ActionTimeTracker.BeginStep(FJoinPartyAttempt::Step_LeaveCurrentParty);
+
+				if (!ExistingParty->IsCurrentlyLeaving())
+				{
+					ExistingParty->LeaveParty(USocialParty::FOnLeavePartyAttemptComplete::CreateUObject(this, &USocialManager::HandleLeavePartyForJoinComplete, ExistingParty));
+				}
+			}
+			else
+			{
+				// This case shouldn't really be possible if the create succeeded - attempt to join.
+				JoinPartyInternal(*JoinAttempt);
+			}
+		}
+		else
+		{
+			// Party creation failed, attempt to join.
+			JoinPartyInternal(*JoinAttempt);
+		}
+	}
 }
 
 void USocialManager::RegisterSocialInteractions()
@@ -1039,7 +1085,17 @@ void USocialManager::HandleQueryJoinabilityComplete(const FUniqueNetId& LocalUse
 			}
 			else
 			{
-				JoinPartyInternal(*JoinAttempt);
+				if (bCreatingPersistentParty && PartyTypeId == IOnlinePartySystem::GetPrimaryPartyTypeId())
+				{
+					UE_LOG(LogParty, Log, TEXT("HandleQueryJoinabilityComplete: User=[%s] Party=[%s] Delaying join until persistent party creation has completed."), *LocalUserId.ToDebugString(), *PartyId.ToDebugString());
+
+					// When the internal persistent party creation completes, the new party will be left to process the join request.
+					JoinAttempt->ActionTimeTracker.BeginStep(FJoinPartyAttempt::Step_WaitForPersistentPartyCreation);
+				}
+				else
+				{
+					JoinPartyInternal(*JoinAttempt);
+				}
 			}
 		}
 		else
diff --git a/Engine/Plugins/Online/OnlineFramework/Source/Party/Public/SocialManager.h b/Engine/Plugins/Online/OnlineFramework/Source/Party/Public/SocialManager.h
--- a/Engine/Plugins/Online/OnlineFramework/Source/Party/Public/SocialManager.h
+++ b/Engine/Plugins/Online/OnlineFramework/Source/Party/Public/SocialManager.h
@@ -150,6 +150,7 @@ protected:
 		static const FName Step_LeaveCurrentParty;
 		static const FName Step_JoinParty;
 		static const FName Step_DeferredPartyCreation;
+		static const FName Step_WaitForPersistentPartyCreation;
 
 		FSocialActionTimeTracker ActionTimeTracker;
 	};
@@ -221,6 +222,9 @@ private:
 
 	TSharedPtr<const IOnlinePartyJoinInfo> GetJoinInfoFromSession(const FOnlineSessionSearchResult& PlatformSession);
 
+	void OnCreatePersistentPartyCompleteInternal(ECreatePartyCompletionResult Result, FOnCreatePartyAttemptComplete OnCreatePartyComplete);
+	bool bCreatingPersistentParty = false;
+
 private:	// Handlers
 	void HandleGameViewportInitialized();
 	void HandlePreClientTravel(const FString& PendingURL, ETravelType TravelType, bool bIsSeamlessTravel);
