Add FailureReason to Social.Party.Invite Event

#tests
Added breakpoint in SendingPartyInvite
Launched two clients
Sent successful invite to online friend, verified success reason
I don't have a good way of triggering a failed invite.

diff --git a/Engine/Plugins/Online/OnlineFramework/Source/Party/Private/Party/SocialParty.cpp b/Engine/Plugins/Online/OnlineFramework/Source/Party/Private/Party/SocialParty.cpp
--- a/Engine/Plugins/Online/OnlineFramework/Source/Party/Private/Party/SocialParty.cpp
+++ b/Engine/Plugins/Online/OnlineFramework/Source/Party/Private/Party/SocialParty.cpp
@@ -263,43 +263,44 @@ bool USocialParty::HasUserBeenInvited(const USocialUser& User) const
 
 bool USocialParty::CanInviteUser(const USocialUser& User) const
 {
-	return CanInviteUserInternal(User);
+	return CanInviteUserInternal(User) == ESocialPartyInviteFailureReason::Success;
 }
 
-bool USocialParty::CanInviteUserInternal(const USocialUser& User) const
+ESocialPartyInviteFailureReason USocialParty::CanInviteUserInternal(const USocialUser& User) const
 {
 	// Only users that are online can be invited
 	if (!User.IsOnline() && !User.CanReceiveOfflineInvite())
 	{
-		return false;
+		return ESocialPartyInviteFailureReason::NotOnline;
 	}
 
 	if (!CurrentConfig.bIsAcceptingMembers && CurrentConfig.NotAcceptingMembersReason != (int32)EPartyJoinDenialReason::PartyPrivate)
 	{
 		// We aren't accepting members for a reason other than party privacy, so a direct invite won't help
-		return false;
+		return ESocialPartyInviteFailureReason::NotAcceptingMembers;
 	}
 
 	//@todo DanH Party: The problem with CanLocalUserInvite is that it the "friend" restriction is applied to mcp friends only, so a console friend doesn't count (but should) #required
 	//		Need to check in with OGS about that...
 	if (!OssParty->CanLocalUserInvite(*OwningLocalUserId))
 	{
-		return false;
+		return ESocialPartyInviteFailureReason::OssValidationFailed;
 	}
 
 	if (GetPartyMember(User.GetUserId(ESocialSubsystem::Primary)))
 	{
 		// Already in the party
-		return false;
+		return ESocialPartyInviteFailureReason::AlreadyInParty;
 	}
 
-	return true;
+	return ESocialPartyInviteFailureReason::Success;
 }
 
 bool USocialParty::TryInviteUser(const USocialUser& UserToInvite, const ESocialPartyInviteMethod InviteMethod)
 {
 	bool bSentInvite = false;
-	if (CanInviteUser(UserToInvite))
+	ESocialPartyInviteFailureReason CanInviteResult = CanInviteUserInternal(UserToInvite);
+	if (CanInviteResult == ESocialPartyInviteFailureReason::Success)
 	{
 		const bool bPreferPlatformInvite = USocialSettings::ShouldPreferPlatformInvites();
 		const bool bMustSendPrimaryInvite = USocialSettings::MustSendPrimaryInvites();
@@ -337,7 +338,8 @@ bool USocialParty::TryInviteUser(const USocialUser& UserToInvite, const ESocialP
 					bSentPlatformInvite = PlatformSessionInterface->SendSessionInviteToFriend(*LocalUserPlatformId, NAME_PartySession, *UserPlatformId);
 				}
 			}
-			OnInviteSentInternal(ESocialSubsystem::Platform, UserToInvite, bSentPlatformInvite, InviteMethod);
+			ESocialPartyInviteFailureReason FailureReason = bSentPlatformInvite ? ESocialPartyInviteFailureReason::Success : ESocialPartyInviteFailureReason::PlatformInviteFailed;
+			OnInviteSentInternal(ESocialSubsystem::Platform, UserToInvite, bSentPlatformInvite, FailureReason, InviteMethod);
 			bSentInvite |= bSentPlatformInvite;
 		}
 		if ((!bSentInvite || bMustSendPrimaryInvite) && UserPrimaryId.IsValid())
@@ -345,13 +347,14 @@ bool USocialParty::TryInviteUser(const USocialUser& UserToInvite, const ESocialP
 			// Primary subsystem invites can be sent directly to the user via the party interface
 			const IOnlinePartyPtr PartyInterface = Online::GetPartyInterfaceChecked(GetWorld());
 			const bool bSentPrimaryInvite = PartyInterface->SendInvitation(*OwningLocalUserId, GetPartyId(), *UserPrimaryId);
-			OnInviteSentInternal(ESocialSubsystem::Primary, UserToInvite, bSentPrimaryInvite, InviteMethod);
+			ESocialPartyInviteFailureReason FailureReason = bSentPrimaryInvite ? ESocialPartyInviteFailureReason::Success : ESocialPartyInviteFailureReason::PartyInviteFailed;
+			OnInviteSentInternal(ESocialSubsystem::Primary, UserToInvite, bSentPrimaryInvite, FailureReason, InviteMethod);
 			bSentInvite |= bSentPrimaryInvite;
 		}
 	}
 	else
 	{
-		OnInviteSentInternal(ESocialSubsystem::MAX, UserToInvite, false, InviteMethod);
+		OnInviteSentInternal(ESocialSubsystem::MAX, UserToInvite, false, CanInviteResult, InviteMethod);
 	}
 	return bSentInvite;
 }
@@ -629,7 +632,7 @@ void USocialParty::OnLeftPartyInternal(EMemberExitedReason Reason)
 	OnPartyLeft().Broadcast(Reason);
 }
 
-void USocialParty::OnInviteSentInternal(ESocialSubsystem SubsystemType, const USocialUser& InvitedUser, bool bWasSuccessful, const ESocialPartyInviteMethod InviteMethod)
+void USocialParty::OnInviteSentInternal(ESocialSubsystem SubsystemType, const USocialUser& InvitedUser, bool bWasSuccessful, const ESocialPartyInviteFailureReason FailureReason, const ESocialPartyInviteMethod InviteMethod)
 {
 	OnInviteSent().Broadcast(InvitedUser);
 
diff --git a/Engine/Plugins/Online/OnlineFramework/Source/Party/Public/Party/PartyTypes.h b/Engine/Plugins/Online/OnlineFramework/Source/Party/Public/Party/PartyTypes.h
--- a/Engine/Plugins/Online/OnlineFramework/Source/Party/Public/Party/PartyTypes.h
+++ b/Engine/Plugins/Online/OnlineFramework/Source/Party/Public/Party/PartyTypes.h
@@ -225,13 +225,46 @@ enum class EApprovalAction : uint8
 };
 
 UENUM()
-enum class ESocialPartyInviteMethod : uint8 {
+enum class ESocialPartyInviteMethod : uint8 
+{
 	/** Default value for try invite */
 	Other = 0,
 	/** Invite was sent from a toast */
 	Notification
 };
 
+UENUM()
+enum class ESocialPartyInviteFailureReason : uint8 
+{
+	Success = 0,
+	NotOnline,
+	NotAcceptingMembers,
+	NotFriends,
+	AlreadyInParty,
+	OssValidationFailed,
+	PlatformInviteFailed,
+	PartyInviteFailed
+};
+
+
+inline const TCHAR* LexToString(ESocialPartyInviteFailureReason Type)
+{
+	switch (Type)
+	{
+	case ESocialPartyInviteFailureReason::Success: return TEXT("Success");
+	case ESocialPartyInviteFailureReason::NotOnline: return TEXT("NotOnline");
+	case ESocialPartyInviteFailureReason::NotAcceptingMembers: return TEXT("NotAcceptingMembers");
+	case ESocialPartyInviteFailureReason::NotFriends: return TEXT("NotFriends");
+	case ESocialPartyInviteFailureReason::AlreadyInParty: return TEXT("AlreadyInParty");
+	case ESocialPartyInviteFailureReason::OssValidationFailed: return TEXT("OssValidationFailed");
+	case ESocialPartyInviteFailureReason::PlatformInviteFailed: return TEXT("PlatformInviteFailed");
+	case ESocialPartyInviteFailureReason::PartyInviteFailed: return TEXT("PartyInviteFailed");
+	default:
+		checkNoEntry();
+		return TEXT("Unknown");
+	}
+}
+
 // Gives a smidge more meaning to the intended use for the string. These should just be UniqueId's (and are), but not reliably allocated as shared ptrs, so they cannot be replicated via FUniqueNetIdRepl.
 using FSessionId = FString;
 
diff --git a/Engine/Plugins/Online/OnlineFramework/Source/Party/Public/Party/SocialParty.h b/Engine/Plugins/Online/OnlineFramework/Source/Party/Public/Party/SocialParty.h
--- a/Engine/Plugins/Online/OnlineFramework/Source/Party/Public/Party/SocialParty.h
+++ b/Engine/Plugins/Online/OnlineFramework/Source/Party/Public/Party/SocialParty.h
@@ -245,11 +245,11 @@ protected:
 	virtual void OnLeftPartyInternal(EMemberExitedReason Reason);
 
 	/** Virtual versions of the package-scoped "CanX" methods above, as a virtual declared within package scoping cannot link (exported public, imported protected) */
-	virtual bool CanInviteUserInternal(const USocialUser& User) const;
+	virtual ESocialPartyInviteFailureReason CanInviteUserInternal(const USocialUser& User) const;
 	virtual bool CanPromoteMemberInternal(const UPartyMember& PartyMember) const;
 	virtual bool CanKickMemberInternal(const UPartyMember& PartyMember) const;
 
-	virtual void OnInviteSentInternal(ESocialSubsystem SubsystemType, const USocialUser& InvitedUser, bool bWasSuccessful, const ESocialPartyInviteMethod InviteMethod);
+	virtual void OnInviteSentInternal(ESocialSubsystem SubsystemType, const USocialUser& InvitedUser, bool bWasSuccessful, const ESocialPartyInviteFailureReason FailureReason, const ESocialPartyInviteMethod InviteMethod);
 
 	/* Deprecated version */
 	virtual void OnInviteSentInternal(ESocialSubsystem SubsystemType, const USocialUser& InvitedUser, bool bWasSuccessful) {};
