Integrate CL#16396842:

Asynchronously flush loading before hotfix / content patching. Synchronously flush loading before swapping config files in GConfig.

diff --git a/Engine/Plugins/Online/OnlineFramework/Source/Hotfix/Private/OnlineHotfixManager.cpp b/Engine/Plugins/Online/OnlineFramework/Source/Hotfix/Private/OnlineHotfixManager.cpp
--- a/Engine/Plugins/Online/OnlineFramework/Source/Hotfix/Private/OnlineHotfixManager.cpp
+++ b/Engine/Plugins/Online/OnlineFramework/Source/Hotfix/Private/OnlineHotfixManager.cpp
@@ -23,6 +23,8 @@
 #include "Curves/CurveFloat.h"
 #include "Engine/BlueprintGeneratedClass.h"
 
+#include "Serialization/AsyncLoadingFlushContext.h"
+
 #ifdef WITH_ONLINETRACING
 #include "OnlineTracingModule.h"
 #endif
@@ -177,6 +179,15 @@ UOnlineHotfixManager::UOnlineHotfixManager() :
 	GameContentPath = FString() / FApp::GetProjectName() / TEXT("Content");
 }
 
+UOnlineHotfixManager::UOnlineHotfixManager(FVTableHelper& Helper)
+	: Super(Helper)
+{
+}
+
+UOnlineHotfixManager::~UOnlineHotfixManager()
+{
+}
+
 UOnlineHotfixManager* UOnlineHotfixManager::Get(UWorld* World)
 {
 	UOnlineHotfixManager* DefaultObject = UOnlineHotfixManager::StaticClass()->GetDefaultObject<UOnlineHotfixManager>();
@@ -258,6 +269,7 @@ void UOnlineHotfixManager::Cleanup()
 	}
 	OnlineTitleFile = nullptr;
 	bHotfixingInProgress = false;
+	AsyncFlushContext = nullptr;
 }
 
 void UOnlineHotfixManager::StartHotfixProcess()
@@ -419,14 +431,29 @@ void UOnlineHotfixManager::OnEnumerateFilesComplete(bool bWasSuccessful, const F
 		{
 			if (RemovedHotfixFileList.Num() > 0)
 			{
-				// No changes, just reverts
-				// Perform any undo operations needed
-				RestoreBackupIniFiles();
-				UnmountHotfixFiles();
+				UE_LOG(LogHotfixManager, Display, TEXT("Files have been removed since last check. Reverting."));
+
+				// Prevent async loading while reverting hotfixes.
+				check(!AsyncFlushContext);
+				AsyncFlushContext = MakeUnique<FAsyncLoadingFlushContext>(TEXT("RevertHotfix"));
+				AsyncFlushContext->Flush(
+					FOnAsyncLoadingFlushComplete::CreateWeakLambda(
+						this,
+						[this]()
+						{
+							// No changes, just reverts
+							// Perform any undo operations needed
+							RestoreBackupIniFiles();
+							UnmountHotfixFiles();
+							
+							TriggerHotfixComplete(EHotfixResult::SuccessNoChange);
+						}));
+			}
+			else
+			{
+				UE_LOG(LogHotfixManager, Display, TEXT("Returned hotfix data is the same as last application, skipping the apply phase"));
+				TriggerHotfixComplete(EHotfixResult::SuccessNoChange);
 			}
-
-			UE_LOG(LogHotfixManager, Display, TEXT("Returned hotfix data is the same as last application, skipping the apply phase"));
-			TriggerHotfixComplete(EHotfixResult::SuccessNoChange);
 		}
 	}
 	else
@@ -648,8 +675,17 @@ void UOnlineHotfixManager::OnReadFileComplete(bool bWasSuccessful, const FString
 			PendingHotfixFiles.Remove(FileName);
 			if (PendingHotfixFiles.Num() == 0)
 			{
-				// All files have been downloaded so now apply the files
-				ApplyHotfix();
+				// Prevent async loading while applying hotfixes.
+				check(!AsyncFlushContext);
+				AsyncFlushContext = MakeUnique<FAsyncLoadingFlushContext>(TEXT("ApplyHotfix"));
+				AsyncFlushContext->Flush(
+					FOnAsyncLoadingFlushComplete::CreateWeakLambda(
+						this,
+						[this]()
+						{
+							const EHotfixResult Result = ApplyHotfix();
+							TriggerHotfixComplete(Result);
+						}));
 			}
 		}
 		else
@@ -668,7 +704,7 @@ void UOnlineHotfixManager::UpdateProgress(uint32 FileCount, uint64 UpdateSize)
 	TriggerOnHotfixProgressDelegates(NumDownloaded, TotalFiles, NumBytes, TotalBytes);
 }
 
-void UOnlineHotfixManager::ApplyHotfix()
+EHotfixResult UOnlineHotfixManager::ApplyHotfix()
 {
 	// Perform any undo operations needed
 	// This occurs same frame as the application of new hotfixes
@@ -680,8 +716,7 @@ void UOnlineHotfixManager::ApplyHotfix()
 		if (!ApplyHotfixProcessing(FileHeader))
 		{
 			UE_LOG(LogHotfixManager, Error, TEXT("Couldn't apply hotfix file (%s)"), *FileHeader.FileName);
-			TriggerHotfixComplete(EHotfixResult::Failed);
-			return;
+			return EHotfixResult::Failed;
 		}
 		// Let anyone listening know we just processed this file
 		TriggerOnHotfixProcessedFileDelegates(FileHeader.FileName, GetCachedDirectory() / FileHeader.DLName);
@@ -698,7 +733,8 @@ void UOnlineHotfixManager::ApplyHotfix()
 		UE_LOG(LogHotfixManager, Display, TEXT("Hotfix has detected PAK files containing currently loaded maps, so a level load is needed"));
 		Result = EHotfixResult::SuccessNeedsReload;
 	}
-	TriggerHotfixComplete(Result);
+
+	return Result;
 }
 
 void UOnlineHotfixManager::TriggerHotfixComplete(EHotfixResult HotfixResult)
@@ -861,6 +897,9 @@ bool UOnlineHotfixManager::HotfixIniFile(const FString& FileName, const FString&
 {
 	const bool bIsEngineIni = FileName.Contains(TEXT("Engine.ini"));
 
+	// Flush async loading before modifying GConfig.
+	FlushAsyncLoading();
+
 	FConfigFile* ConfigFile = GetConfigFile(FileName);
 	// Store the original file so we can undo this later
 	FConfigFileBackup& BackupFile = BackupIniFile(FileName, ConfigFile);
@@ -1176,6 +1215,9 @@ bool UOnlineHotfixManager::IsMapLoaded(const FString& MapName)
 
 bool UOnlineHotfixManager::HotfixPakIniFile(const FString& FileName)
 {
+	// Flush async loading before modifying GConfig.
+	FlushAsyncLoading();
+
 	FString StrippedName;
 	const double StartTime = FPlatformTime::Seconds();
 	// Need to strip off the PAK path
@@ -1317,6 +1359,10 @@ void UOnlineHotfixManager::RestoreBackupIniFiles()
 	{
 		return;
 	}
+
+	// Flush async loading before modifying GConfig.
+	FlushAsyncLoading();
+
 	const double StartTime = FPlatformTime::Seconds();
 	TArray<FString> ClassesToRestore;
 
@@ -1403,6 +1449,9 @@ void UOnlineHotfixManager::PatchAssetsFromIniFiles()
 {
 	UE_LOG(LogHotfixManager, Display, TEXT("Checking for assets to be patched using data from 'AssetHotfix' section in the Game .ini file"));
 
+	// Flush async loading before modifying GConfig.
+	FlushAsyncLoading();
+
 	int32 TotalPatchableAssets = 0;
 	AssetsHotfixedFromIniFiles.Reset();
 
diff --git a/Engine/Plugins/Online/OnlineFramework/Source/Hotfix/Public/OnlineHotfixManager.h b/Engine/Plugins/Online/OnlineFramework/Source/Hotfix/Public/OnlineHotfixManager.h
--- a/Engine/Plugins/Online/OnlineFramework/Source/Hotfix/Public/OnlineHotfixManager.h
+++ b/Engine/Plugins/Online/OnlineFramework/Source/Hotfix/Public/OnlineHotfixManager.h
@@ -15,6 +15,8 @@
 
 HOTFIX_API DECLARE_LOG_CATEGORY_EXTERN(LogHotfixManager, Display, All);
 
+class FAsyncLoadingFlushContext;
+
 UENUM()
 enum class EHotfixResult : uint8
 {
@@ -157,6 +159,8 @@ protected:
 	uint64 NumBytes;
 	/** Some title file interfaces aren't re-entrant so handle it ourselves */
 	bool bHotfixingInProgress;
+	/** Asynchronously flush async loading before starting the hotfixing process. */
+	TUniquePtr<FAsyncLoadingFlushContext> AsyncFlushContext;
 	/** Set to true if any PAK file contains an update to a level that is currently loaded */
 	bool bHotfixNeedsMapReload;
 #if !UE_BUILD_SHIPPING
@@ -174,7 +178,7 @@ protected:
 	virtual void Init();
 	virtual void Cleanup();
 	/** Looks at each file returned via the hotfix and processes them */
-	void ApplyHotfix();
+	EHotfixResult ApplyHotfix();
 	/** Cleans up and fires the delegate indicating it's done */
 	void TriggerHotfixComplete(EHotfixResult HotfixResult);
 	/** Checks each file listed to see if it is a hotfix file to process */
@@ -311,6 +315,8 @@ protected:
 
 public:
 	UOnlineHotfixManager();
+	UOnlineHotfixManager(FVTableHelper& Helper);
+	virtual ~UOnlineHotfixManager();
 
 	/** Tells the hotfix manager which OSS to use. Uses the default if empty */
 	UPROPERTY(Config)
diff --git a/Engine/Source/Runtime/CoreUObject/Private/Serialization/AsyncLoadingFlushContext.cpp b/Engine/Source/Runtime/CoreUObject/Private/Serialization/AsyncLoadingFlushContext.cpp
new file mode 100644
--- /dev/null
+++ b/Engine/Source/Runtime/CoreUObject/Private/Serialization/AsyncLoadingFlushContext.cpp
@@ -0,0 +1,81 @@
+// Copyright Epic Games, Inc. All Rights Reserved.
+
+#include "Serialization/AsyncLoadingFlushContext.h"
+
+#include "Containers/Ticker.h"
+#include "Serialization/AsyncPackageLoader.h"
+
+DEFINE_LOG_CATEGORY(AsyncLoadingFlush);
+
+bool IsAsyncLoadingCoreUObjectInternal();
+
+int32 FAsyncLoadingFlushContext::NextId = 0;
+
+FAsyncLoadingFlushContext::FAsyncLoadingFlushContext(const FString& Context)
+	: Context(Context)
+	, Id(NextId++)
+{
+}
+
+FAsyncLoadingFlushContext::FAsyncLoadingFlushContext(FString&& Context)
+	: Context(MoveTemp(Context))
+	, Id(NextId++)
+{
+}
+
+FAsyncLoadingFlushContext::~FAsyncLoadingFlushContext()
+{
+	CleanupTickers();
+}
+
+void FAsyncLoadingFlushContext::Flush(const FOnAsyncLoadingFlushComplete& OnFlushComplete)
+{
+	check(IsInGameThread() && !IsInSlateThread());
+
+	OnFlushCompleteDelegate = OnFlushComplete;
+	check(OnFlushCompleteDelegate.IsBound());
+
+	// Check each frame whether async loading has completed.
+	LoadingCompleteCheckDelegateHandle = FTicker::GetCoreTicker().AddTicker(FTickerDelegate::CreateRaw(this, &FAsyncLoadingFlushContext::OnAsyncLoadingCheck));
+
+	// Periodically warn when waiting for async loading to complete.
+	WaitingWarnDelegateHandle = FTicker::GetCoreTicker().AddTicker(FTickerDelegate::CreateRaw(this, &FAsyncLoadingFlushContext::OnAsyncLoadingWarn), LoadingFlushWarnIntervalSeconds);
+
+	StartTime = FPlatformTime::Seconds();
+	UE_LOG(AsyncLoadingFlush, Verbose, TEXT("Context[%s:%d] Flushing async loading."), *Context, Id);
+}
+
+void FAsyncLoadingFlushContext::CleanupTickers()
+{
+	FTicker::GetCoreTicker().RemoveTicker(LoadingCompleteCheckDelegateHandle);
+	FTicker::GetCoreTicker().RemoveTicker(WaitingWarnDelegateHandle);
+}
+
+bool FAsyncLoadingFlushContext::OnAsyncLoadingCheck(float DeltaTime)
+{
+	if (IsAsyncLoadingCoreUObjectInternal())
+	{
+		// Continue rescheduling the check.
+		return true;
+	}
+
+	// Async loading completed, remove warning logger.
+	CleanupTickers();
+
+	const double CurrentTimeSeconds = FPlatformTime::Seconds();
+	UE_LOG(AsyncLoadingFlush, Verbose, TEXT("Context[%s:%d] Flushed in %f seconds."), *Context, Id, CurrentTimeSeconds - StartTime);
+
+	// Signal to user that flush has completed.
+	FOnAsyncLoadingFlushComplete OnFlushComplete = MoveTemp(OnFlushCompleteDelegate);
+	OnFlushComplete.ExecuteIfBound();
+	// do not access member variables below this point in case firing the delegate deleted this object.
+
+	// Unregister check.
+	return false;
+}
+
+bool FAsyncLoadingFlushContext::OnAsyncLoadingWarn(float DeltaTime)
+{
+	UE_LOG(AsyncLoadingFlush, Warning, TEXT("Context[%s:%d] Waiting on async loading to complete. Total wait time: %f seconds."), *Context, Id, FPlatformTime::Seconds() - StartTime);
+	return true;
+}
diff --git a/Engine/Source/Runtime/CoreUObject/Public/Serialization/AsyncLoadingFlushContext.h b/Engine/Source/Runtime/CoreUObject/Public/Serialization/AsyncLoadingFlushContext.h
new file mode 100644
--- /dev/null
+++ b/Engine/Source/Runtime/CoreUObject/Public/Serialization/AsyncLoadingFlushContext.h
@@ -0,0 +1,55 @@
+// Copyright Epic Games, Inc. All Rights Reserved.
+
+#pragma once
+
+#include "Delegates/Delegate.h"
+#include "Logging/LogMacros.h"
+
+DECLARE_LOG_CATEGORY_EXTERN(AsyncLoadingFlush, Log, All);
+
+DECLARE_DELEGATE(FOnAsyncLoadingFlushComplete);
+
+/**
+ * Flush the async loader in a non-blocking manner.
+ */
+class COREUOBJECT_API FAsyncLoadingFlushContext final : FNoncopyable
+{
+public:
+	FAsyncLoadingFlushContext(const FString & Context);
+	FAsyncLoadingFlushContext(FString&& Context);
+	~FAsyncLoadingFlushContext();
+
+	void Flush(const FOnAsyncLoadingFlushComplete& OnFlushComplete);
+
+	int32 GetId() const { return Id; }
+
+private:
+	FAsyncLoadingFlushContext() = delete;
+
+	void CleanupTickers();
+
+	bool OnAsyncLoadingCheck(float DeltaTime);
+	bool OnAsyncLoadingWarn(float DeltaTime);
+
+	/** How often to warn while waiting for async loading to complete. */
+	static constexpr float LoadingFlushWarnIntervalSeconds = 10.f;
+
+	/** Context for logging */
+	FString Context;
+
+	/** Unique id for context. */
+	int32 Id;
+
+	/** Delegate to be fired on flush. */
+	FOnAsyncLoadingFlushComplete OnFlushCompleteDelegate;
+
+	/** Record event timing for logging. */
+	double StartTime = 0;
+
+	/** Tick delegates. */
+	FDelegateHandle LoadingCompleteCheckDelegateHandle;
+	FDelegateHandle WaitingWarnDelegateHandle;
+
+	/** Increment ids so each instance has a unique id. */
+	static int32 NextId;
+};
