Social - Added FriendRequestMethod to Social.Friend.Invite Event.  Fixed Add Friend pedestal nameplate button so that it will not display if local player has a pending sent invite for that player.
#TESTS Used multiple clients to send friend requests from various methods in the front end and used breakpoints to captured incoming calls to SendingFriendRequest to verify new FriendRequestMethod parameter was coming in as expected.

diff --git a/Engine/Plugins/Online/OnlineFramework/Source/Party/Private/Interactions/CoreInteractions.cpp b/Engine/Plugins/Online/OnlineFramework/Source/Party/Private/Interactions/CoreInteractions.cpp
--- a/Engine/Plugins/Online/OnlineFramework/Source/Party/Private/Interactions/CoreInteractions.cpp
+++ b/Engine/Plugins/Online/OnlineFramework/Source/Party/Private/Interactions/CoreInteractions.cpp
@@ -38,7 +38,7 @@ bool FSocialInteraction_AddFriend::CanExecute(const USocialUser& User)
 
 void FSocialInteraction_AddFriend::ExecuteInteraction(USocialUser& User)
 {
-	User.SendFriendInvite(ESocialSubsystem::Primary);
+	User.SendFriendInvite(ESocialSubsystem::Primary, ESocialFriendRequestMethod::SocialInteraction_AddFriend);
 }
 
 //////////////////////////////////////////////////////////////////////////
@@ -75,7 +75,7 @@ bool FSocialInteraction_AddPlatformFriend::CanExecute(const USocialUser& User)
 
 void FSocialInteraction_AddPlatformFriend::ExecuteInteraction(USocialUser& User)
 {
-	User.SendFriendInvite(ESocialSubsystem::Platform);
+	User.SendFriendInvite(ESocialSubsystem::Platform, ESocialFriendRequestMethod::SocialInteraction_AddPlatformFriend);
 }
 
 //////////////////////////////////////////////////////////////////////////
diff --git a/Engine/Plugins/Online/OnlineFramework/Source/Party/Private/User/SocialUser.cpp b/Engine/Plugins/Online/OnlineFramework/Source/Party/Private/User/SocialUser.cpp
--- a/Engine/Plugins/Online/OnlineFramework/Source/Party/Private/User/SocialUser.cpp
+++ b/Engine/Plugins/Online/OnlineFramework/Source/Party/Private/User/SocialUser.cpp
@@ -992,7 +992,7 @@ FString USocialUser::ToDebugString() const
 #endif
 }
 
-bool USocialUser::SendFriendInvite(ESocialSubsystem SubsystemType)
+bool USocialUser::SendFriendInvite(ESocialSubsystem SubsystemType, const ESocialFriendRequestMethod RequestMethod /* = ESocialFriendRequestMethod::SocialInteraction_AddFriend */)
 {
 	return GetOwningToolkit().TrySendFriendInvite(*this, SubsystemType);
 }
diff --git a/Engine/Plugins/Online/OnlineFramework/Source/Party/Public/User/SocialUser.h b/Engine/Plugins/Online/OnlineFramework/Source/Party/Public/User/SocialUser.h
--- a/Engine/Plugins/Online/OnlineFramework/Source/Party/Public/User/SocialUser.h
+++ b/Engine/Plugins/Online/OnlineFramework/Source/Party/Public/User/SocialUser.h
@@ -25,6 +25,34 @@ typedef TSharedPtr<const IOnlinePartyJoinInfo> IOnlinePartyJoinInfoConstPtr;
 
 namespace EOnlinePresenceState { enum Type : uint8; }
 
+/** Represents the different methods of sending a friend request */
+UENUM()
+enum class ESocialFriendRequestMethod : uint8 
+{
+	/** Friend Request was sent using SocialInteraction_AddFriend */
+	SocialInteraction_AddFriend,
+	/** Friend Request was sent using SocialInteraction_AddPlatformFriend */
+	SocialInteraction_AddPlatformFriend,
+	/** Other implementation **/
+	Other
+};
+
+/** Helper function that provides readable text value for ESocialFriendRequestMethod enum */
+inline const TCHAR* LexToString(ESocialFriendRequestMethod RequestMethod)
+{
+	switch (RequestMethod)
+	{
+	case ESocialFriendRequestMethod::SocialInteraction_AddFriend:
+		return TEXT("SocialInteraction_AddFriend");
+	case ESocialFriendRequestMethod::SocialInteraction_AddPlatformFriend:
+		return TEXT("SocialInteraction_AddPlatformFriend");
+	case ESocialFriendRequestMethod::Other:
+		return TEXT("Other");
+	default:
+		return TEXT("Invalid SocialFriendRequestMethod");
+	}
+}
+
 DECLARE_DELEGATE_OneParam(FOnNewSocialUserInitialized, USocialUser&);
 
 UCLASS(Within = SocialToolkit)
@@ -93,7 +121,7 @@ public:
 	TArray<FSocialInteractionHandle> GetAllAvailableInteractions() const;
 
 	virtual bool CanSendFriendInvite(ESocialSubsystem SubsystemType) const;
-	virtual bool SendFriendInvite(ESocialSubsystem SubsystemType);
+	virtual bool SendFriendInvite(ESocialSubsystem SubsystemType, const ESocialFriendRequestMethod RequestMethod = ESocialFriendRequestMethod::SocialInteraction_AddFriend);
 	virtual bool AcceptFriendInvite(ESocialSubsystem SocialSubsystem) const;
 	virtual bool RejectFriendInvite(ESocialSubsystem SocialSubsystem) const;
 	virtual bool EndFriendship(ESocialSubsystem SocialSubsystem) const;
