Fix for erroneous dereference of ToString (const TCHAR*) result passed to VeryVerbose UE_LOG when sending a request to join.

diff --git a/Engine/Plugins/Online/OnlineFramework/Source/Party/Private/SocialToolkit.cpp b/Engine/Plugins/Online/OnlineFramework/Source/Party/Private/SocialToolkit.cpp
--- a/Engine/Plugins/Online/OnlineFramework/Source/Party/Private/SocialToolkit.cpp
+++ b/Engine/Plugins/Online/OnlineFramework/Source/Party/Private/SocialToolkit.cpp
@@ -1157,7 +1157,7 @@ void USocialToolkit::RequestToJoinParty(USocialUser& SocialUser)
 
 void USocialToolkit::HandlePartyRequestToJoinSent(const FUniqueNetId& LocalUserId, const FUniqueNetId& PartyLeaderId, const FDateTime& ExpiresAt, const ERequestToJoinPartyCompletionResult Result)
 {
-	UE_LOG(LogParty, VeryVerbose, TEXT("%s - User [%s] sent a join request to [%s] with result[%s]"), ANSI_TO_TCHAR(__FUNCTION__), *LocalUserId.ToDebugString(), *PartyLeaderId.ToDebugString(), *ToString(Result));
+	UE_LOG(LogParty, VeryVerbose, TEXT("%s - User [%s] sent a join request to [%s] with result[%s]"), ANSI_TO_TCHAR(__FUNCTION__), *LocalUserId.ToDebugString(), *PartyLeaderId.ToDebugString(), ToString(Result));
 
 	if (Result == ERequestToJoinPartyCompletionResult::Succeeded)
 	{
