Added rules-based ordering of QOS subregions into subspaces, biasing datacenter selection away from subspaces whose performance is not a sufficient improvement over the containing subregion.  The tolerances to control this can be set via the Qos.QosRegionManager section in the engine config on a per-RegionDefinition basis, and also enabled/disabled per-region.  The rules-based comparison may be enabled/disabled overall via the engine config, or overridden with the `-qossubspacebias=true|false` command-line argument.

diff --git a/Engine/Plugins/Online/OnlineFramework/Source/Qos/Private/QosRegionManager.cpp b/Engine/Plugins/Online/OnlineFramework/Source/Qos/Private/QosRegionManager.cpp
--- a/Engine/Plugins/Online/OnlineFramework/Source/Qos/Private/QosRegionManager.cpp
+++ b/Engine/Plugins/Online/OnlineFramework/Source/Qos/Private/QosRegionManager.cpp
@@ -11,6 +11,73 @@
 
 #define LAST_REGION_EVALUATION 3
 
+bool FDatacenterQosInstance::IsLessWhenBiasedTowardsNonSubspace(
+	const FDatacenterQosInstance& A,
+	const FDatacenterQosInstance& B,
+	const FQosSubspaceComparisonParams& ComparisonParams,
+	const TCHAR* const SubspaceDelimiter)
+{
+	if (A.Definition.IsSubspace(SubspaceDelimiter) == B.Definition.IsSubspace(SubspaceDelimiter))
+	{
+		// Comparing like-for-like, so use normal ordering by average ping.
+		return A.AvgPingMs < B.AvgPingMs;
+	}
+
+	// Comparing subspace vs non-subspace, or vice-versa.
+	const bool bAIsSubspace = A.Definition.IsSubspace(SubspaceDelimiter);
+	const FDatacenterQosInstance& Subspace = bAIsSubspace ? A : B;
+	const FDatacenterQosInstance& NonSubspace = bAIsSubspace ? B : A;
+	const bool bIsNonSubspaceRecommended = IsNonSubspaceRecommended(NonSubspace, Subspace, ComparisonParams);
+
+	const bool bResult = bIsNonSubspaceRecommended != bAIsSubspace;
+	return bResult;
+}
+
+bool FDatacenterQosInstance::IsNonSubspaceRecommended(
+	const FDatacenterQosInstance& NonSubspace,
+	const FDatacenterQosInstance& Subspace,
+	const FQosSubspaceComparisonParams& ComparisonParams)
+{
+	const int32 NonSubspacePingMs = NonSubspace.AvgPingMs;
+	const int32 SubspacePingMs = Subspace.AvgPingMs;
+
+	if (NonSubspacePingMs < SubspacePingMs)
+	{
+		// Edge case where non-subspace is faster, recommend it immediately.
+		return true;
+	}
+
+	const int32 PingMsDelta = NonSubspacePingMs - SubspacePingMs;
+
+	const float ScaledMaxToleranceMs = ComparisonParams.CalcScaledMaxToleranceMs(NonSubspacePingMs);
+
+	// Test if non-subspace qualifies over subspace, via rules-based comparison that biases towards non-subspace.
+	bool bDisqualify = (ComparisonParams.MaxNonSubspacePingMs > 0) && (NonSubspacePingMs > ComparisonParams.MaxNonSubspacePingMs);
+	bDisqualify = bDisqualify || ((ComparisonParams.MinSubspacePingMs > 0) && (SubspacePingMs < ComparisonParams.MinSubspacePingMs));
+	bDisqualify = bDisqualify || ((ComparisonParams.ConstantMaxToleranceMs > 0) && (PingMsDelta > ComparisonParams.ConstantMaxToleranceMs));
+	bDisqualify = bDisqualify || ((ComparisonParams.ScaledMaxTolerancePct > 0.0f) && (static_cast<float>(PingMsDelta) > ScaledMaxToleranceMs));
+
+	return !bDisqualify;
+}
+
+const FDatacenterQosInstance& FDatacenterQosInstance::CompareBiasedTowardsNonSubspace(
+	const FDatacenterQosInstance& NonSubspace, const FDatacenterQosInstance& Subspace,
+	const FQosSubspaceComparisonParams& ComparisonParams)
+{
+	const bool bPreferNonSubspace = IsNonSubspaceRecommended(NonSubspace, Subspace, ComparisonParams);
+	return bPreferNonSubspace ? NonSubspace : Subspace;
+}
+
+const FDatacenterQosInstance* const FDatacenterQosInstance::CompareBiasedTowardsNonSubspace(
+	const FDatacenterQosInstance* const NonSubspace, const FDatacenterQosInstance* const Subspace,
+	const FQosSubspaceComparisonParams& ComparisonParams)
+{
+	check(NonSubspace);
+	check(Subspace);
+	const bool bPreferNonSubspace = IsNonSubspaceRecommended(*NonSubspace, *Subspace, ComparisonParams);
+	return bPreferNonSubspace ? NonSubspace : Subspace;
+}
+
 EQosDatacenterResult FRegionQosInstance::GetRegionResult() const
 {
 	EQosDatacenterResult Result = EQosDatacenterResult::Success;
@@ -48,7 +115,7 @@ FString FRegionQosInstance::GetBestSubregion() const
 	FString BestDatacenterId;
 	if (DatacenterOptions.Num() > 0)
 	{
-		// Presorted for best result first
+		// Presorted for best/recommended result first
 		BestDatacenterId = DatacenterOptions[0].Definition.Id;
 	}
 
@@ -57,16 +124,40 @@ FString FRegionQosInstance::GetBestSubregion() const
 
 void FRegionQosInstance::GetSubregionPreferences(TArray<FString>& OutSubregions) const
 {
+	// Presorted for best/recommended result first
 	for (const FDatacenterQosInstance& Option : DatacenterOptions)
 	{
-		// Presorted for best result first
 		OutSubregions.Add(Option.Definition.Id);
 	}
 }
 
+void FRegionQosInstance::SortDatacenterOptionsByAvgPingAsc()
+{
+	// Sort ping best to worst (ascending ping ms)
+	DatacenterOptions.Sort(
+		[](const FDatacenterQosInstance& A, const FDatacenterQosInstance& B) {
+			return A.AvgPingMs < B.AvgPingMs;
+		}
+	);
+}
+
+void FRegionQosInstance::SortDatacenterSubspacesByRecommended(const FQosSubspaceComparisonParams& ComparisonParams, const TCHAR* const SubspaceDelimiter)
+{
+	// Sort by subspace recommendation rules
+	DatacenterOptions.Sort(
+		[&ComparisonParams, &SubspaceDelimiter](const FDatacenterQosInstance& A, const FDatacenterQosInstance& B) {
+			return FDatacenterQosInstance::IsLessWhenBiasedTowardsNonSubspace(A, B, ComparisonParams, SubspaceDelimiter);
+		}
+	);
+}
+
+const TCHAR* UQosRegionManager::SubspaceDelimiterDefault = TEXT("_");
+
 UQosRegionManager::UQosRegionManager(const FObjectInitializer& ObjectInitializer)
 	: NumTestsPerRegion(3)
 	, PingTimeout(5.0f)
+	, bEnableSubspaceBiasOrder(false)
+	, SubspaceDelimiter(SubspaceDelimiterDefault)
 	, LastCheckTimestamp(0)
 	, Evaluator(nullptr)
 	, QosEvalResult(EQosCompletionResult::Invalid)
@@ -156,6 +247,24 @@ int32 UQosRegionManager::GetMaxPingMs() const
 	return -1;
 }
 
+bool UQosRegionManager::IsSubspaceBiasOrderEnabled() const
+{
+	// Command-line override, if specified.
+	bool bEnabledCommandLine = false;
+	if (FParse::Bool(FCommandLine::Get(), TEXT("qossubspacebias="), bEnabledCommandLine))
+	{
+		return bEnabledCommandLine;
+	}
+
+	// Use config value (or internal code default)
+	return bEnableSubspaceBiasOrder;
+}
+
+bool UQosRegionManager::IsSubspaceBiasOrderEnabled(const FQosRegionInfo& RegionDefinition) const
+{
+	return IsSubspaceBiasOrderEnabled() && RegionDefinition.bAllowSubspaceBias;
+}
+
 // static
 FString UQosRegionManager::GetDatacenterId()
 {
@@ -286,6 +395,20 @@ void UQosRegionManager::OnQosEvaluationComplete(EQosCompletionResult Result, con
 		DatacenterMap.Add(Datacenter.Definition.RegionId, Datacenter);
 	}
 
+	const bool bSubspaceBiasEnabled = IsSubspaceBiasOrderEnabled();
+	const TCHAR* const Delim = GetSubspaceDelimiter();
+
+#if DEBUG_SUBCOMPARE_BY_SUBSPACE
+	// Test example regions (indepedent of real obtained QoS data)
+	{
+		TestCompareDatacentersBySubspace();
+		TestSortDatacenterSubspacesByRecommended();
+
+		FRegionQosInstance TestRegion = TestCreateExampleRegionResult();
+		TestRegion.SortDatacenterSubspacesByRecommended(TestRegion.Definition.SubspaceBiasParams, TEXT("_"));
+	}
+#endif // DEBUG_SUBCOMPARE_BY_SUBSPACE
+
 	for (const FQosRegionInfo& RegionInfo : RegionDefinitions)
 	{
 		if (RegionInfo.IsPingable())
@@ -296,11 +419,19 @@ void UQosRegionManager::OnQosEvaluationComplete(EQosCompletionResult Result, con
 				FRegionQosInstance* NewRegion = new (RegionOptions) FRegionQosInstance(RegionInfo);
 				DatacenterMap.MultiFind(RegionInfo.RegionId, NewRegion->DatacenterOptions);
 
-				NewRegion->DatacenterOptions.Sort([](const FDatacenterQosInstance& A, const FDatacenterQosInstance& B)
+				const bool bSubspaceBias = bSubspaceBiasEnabled && RegionInfo.bAllowSubspaceBias;
+				if (bSubspaceBias)
 				{
-					// Sort ping best to worst
-					return A.AvgPingMs < B.AvgPingMs;
-				});
+					// Use ordering that applies rules-based comparison for subspace vs non-subspace,
+					// and avg ping when comparing like-for-like.
+					const FQosSubspaceComparisonParams ComparisonParams = RegionInfo.SubspaceBiasParams;
+					NewRegion->SortDatacenterSubspacesByRecommended(ComparisonParams, Delim);
+				}
+				else
+				{
+					// Use regular sorting of subregions, by ascendering order of avg ping.
+					NewRegion->SortDatacenterOptionsByAvgPingAsc();
+				}
 			}
 			else
 			{
@@ -332,9 +463,14 @@ void UQosRegionManager::OnQosEvaluationComplete(EQosCompletionResult Result, con
 			TrySetDefaultRegion();
 		}
 	}
-	TArray<FString> BestSubregions;
-	GetSubregionPreferences(GetBestRegion(), BestSubregions);
-	UE_LOG(LogQos, Log, TEXT("[UQosRegionManager::OnQosEvaluationComplete] ping eval has completed - our best region is '%s' and our best subregion is '%s'"), *GetBestRegion(), BestSubregions.Num() ? *BestSubregions[0] : TEXT("NONE"));
+	TArray<FString> BestRegionSubregions;
+	GetSubregionPreferences(GetBestRegion(), BestRegionSubregions);
+	UE_LOG(LogQos, Log, TEXT("[UQosRegionManager::OnQosEvaluationComplete] ping eval has completed.  Best region is '%s', recommended subregion is '%s'"),
+		*GetBestRegion(), BestRegionSubregions.Num() ? *BestRegionSubregions[0] : TEXT("NONE"));
+
+#if DEBUG_SUBCOMPARE_BY_SUBSPACE
+	DumpRegionStats();
+#endif // DEBUG_SUBCOMPARE_BY_SUBSPACE
 
 	OnQosEvalCompleteDelegate.Broadcast();
 	OnQosEvalCompleteDelegate.Clear();
@@ -344,7 +480,7 @@ FString UQosRegionManager::GetRegionId() const
 {
 	if (!ForceRegionId.IsEmpty())
 	{
-		UE_LOG(LogQos, VeryVerbose, TEXT("[UQosRegionManager::GetRegionId] Force region: \"%s\""), *ForceRegionId);
+		//UE_LOG(LogQos, VeryVerbose, TEXT("[UQosRegionManager::GetRegionId] Force region: \"%s\""), *ForceRegionId);
 
 		// we may have updated INI to bypass this process
 		return ForceRegionId;
@@ -370,6 +506,41 @@ FString UQosRegionManager::GetRegionId() const
 	return SelectedRegionId;
 }
 
+const TCHAR* UQosRegionManager::GetSubspaceDelimiter() const
+{
+	return SubspaceDelimiter.IsEmpty() ? SubspaceDelimiterDefault : *SubspaceDelimiter;
+}
+
+const FRegionQosInstance* UQosRegionManager::FindQosRegionById(const TArray<FRegionQosInstance>& Regions, const FString& RegionId)
+{
+	const FRegionQosInstance* Found = Regions.FindByPredicate(
+		[&RegionId](const FRegionQosInstance& Elem) {
+			return Elem.Definition.RegionId == RegionId;
+		}
+	);
+
+	return Found;
+}
+
+const FRegionQosInstance* UQosRegionManager::FindBestQosRegion(const TArray<FRegionQosInstance>& Regions)
+{
+	int32 LowestRegionPing = UNREACHABLE_PING;
+	const FRegionQosInstance* BestRegion = nullptr;
+
+	// "Best" average ping is the avg ping of the first datacenter in each region's list of
+	// datacenter options.  Assumes that the datacenter options list is pre-sorted as needed.
+	for (const FRegionQosInstance& Region : Regions)
+	{
+		if (Region.IsAutoAssignable() && Region.GetBestAvgPing() < LowestRegionPing)
+		{
+			LowestRegionPing = Region.GetBestAvgPing();
+			BestRegion = &Region;
+		}
+	}
+
+	return BestRegion;
+}
+
 FString UQosRegionManager::GetBestRegion() const
 {
 	if (!ForceRegionId.IsEmpty())
@@ -379,19 +550,9 @@ FString UQosRegionManager::GetBestRegion() const
 		return ForceRegionId;
 	}
 
-	FString BestRegionId;
-
-	// try to select the lowest ping
-	int32 BestPing = INT_MAX;
 	const TArray<FRegionQosInstance>& LocalRegionOptions = GetRegionOptions();
-	for (const FRegionQosInstance& Region : LocalRegionOptions)
-	{
-		if (Region.IsAutoAssignable() && Region.GetBestAvgPing() < BestPing)
-		{
-			BestPing = Region.GetBestAvgPing();
-			BestRegionId = Region.Definition.RegionId;
-		}
-	}
+	const FRegionQosInstance* BestRegion = FindBestQosRegion(LocalRegionOptions);
+	const FString BestRegionId = BestRegion ? BestRegion->Definition.RegionId : FString();
 
 	UE_LOG(LogQos, Verbose, TEXT("[UQosRegionManager::GetBestRegion] Best region: \"%s\"  (Current selected: \"%s\")"),
 		*BestRegionId, *SelectedRegionId);
@@ -751,3 +912,370 @@ void UQosRegionManager::RegisterQoSSettingsChangedDelegate(const FSimpleDelegate
 	// add to the completion delegate
 	OnQoSSettingsChangedDelegate = OnQoSSettingsChanged;
 }
+
+#if DEBUG_SUBCOMPARE_BY_SUBSPACE
+
+bool UQosRegionManager::TestCompareDatacentersBySubspace()
+{
+	const FQosSubspaceComparisonParams Invariants(300, 8, 25, 50.0f);
+
+	// Parent is a non-subspace datacenter.
+	FQosDatacenterInfo ParentDatacenter;
+	ParentDatacenter.Id = TEXT("DE");
+	ParentDatacenter.RegionId = TEXT("EU");
+	ParentDatacenter.Servers.Add(FQosPingServerInfo{ "100.1.1.1", 2345 });
+	ParentDatacenter.Servers.Add(FQosPingServerInfo{ "100.1.1.2", 2345 });
+	ParentDatacenter.Servers.Add(FQosPingServerInfo{ "100.1.1.3", 2345 });
+
+	// Child is a subspcae datacenter.
+	FQosDatacenterInfo ChildDatacenter;
+	ChildDatacenter.Id = TEXT("DE_S");
+	ChildDatacenter.RegionId = TEXT("EU");
+	ChildDatacenter.Servers.Add(FQosPingServerInfo{ "100.1.2.11", 2345 });
+	ChildDatacenter.Servers.Add(FQosPingServerInfo{ "100.1.2.12", 2345 });
+	ChildDatacenter.Servers.Add(FQosPingServerInfo{ "100.1.2.13", 2345 });
+
+	// Test 1
+	// Parent is faster, somehow (potential edge case)
+	// Expected result: Pick parent. (override succeeds)
+	{
+		FDatacenterQosInstance ChildSubregion;
+		ChildSubregion.Definition = ChildDatacenter;
+		ChildSubregion.AvgPingMs = 68;
+
+		FDatacenterQosInstance ParentSubregion;
+		ParentSubregion.Definition = ParentDatacenter;
+		ParentSubregion.AvgPingMs = 24;
+
+		const FDatacenterQosInstance& RecommendedSubregion = FDatacenterQosInstance::CompareBiasedTowardsNonSubspace(
+			ParentSubregion, ChildSubregion, Invariants);
+		check(&RecommendedSubregion == &ParentSubregion);
+	}
+
+	// Test 2
+	// Parent is too slow.
+	// Expected result: Pick child. (override fails)
+	{
+		FDatacenterQosInstance ChildSubregion;
+		ChildSubregion.Definition = ChildDatacenter;
+		ChildSubregion.AvgPingMs = 48;
+
+		FDatacenterQosInstance ParentSubregion;
+		ParentSubregion.Definition = ParentDatacenter;
+		ParentSubregion.AvgPingMs = 360;
+
+		const FDatacenterQosInstance& RecommendedSubregion = FDatacenterQosInstance::CompareBiasedTowardsNonSubspace(
+			ParentSubregion, ChildSubregion, Invariants);
+		check(&RecommendedSubregion == &ChildSubregion);
+	}
+
+	// Test 3
+	// Parent is fast, but child is faster than minimum ping
+	// Expected result: Pick child. (override fails)
+	{
+		FDatacenterQosInstance ChildSubregion;
+		ChildSubregion.Definition = ChildDatacenter;
+		ChildSubregion.AvgPingMs = 6;
+
+		FDatacenterQosInstance ParentSubregion;
+		ParentSubregion.Definition = ParentDatacenter;
+		ParentSubregion.AvgPingMs = 34;
+
+		const FDatacenterQosInstance& RecommendedSubregion = FDatacenterQosInstance::CompareBiasedTowardsNonSubspace(
+			ParentSubregion, ChildSubregion, Invariants);
+		check(&RecommendedSubregion == &ChildSubregion);
+	}
+
+	// Test 4
+	// Parent is fast enough for consideration, child is slower than minimum ping.
+	// Difference in ping is GREATER THAN allowed constant tolerance in milliseconds.
+	// Expected result: Pick child. (override fails)
+	{
+		FDatacenterQosInstance ChildSubregion;
+		ChildSubregion.Definition = ChildDatacenter;
+		ChildSubregion.AvgPingMs = 48;
+
+		FDatacenterQosInstance ParentSubregion;
+		ParentSubregion.Definition = ParentDatacenter;
+		ParentSubregion.AvgPingMs = 78; // (78 - 48 = 30) > 25
+
+		const FDatacenterQosInstance& RecommendedSubregion = FDatacenterQosInstance::CompareBiasedTowardsNonSubspace(
+			ParentSubregion, ChildSubregion, Invariants);
+		check(&RecommendedSubregion == &ChildSubregion);
+	}
+
+	// If the difference in ping is LESS THAN or EQUAL TO allowed constant tolerance
+	// in milliseconds, then we test against proportional tolerance (below).
+
+	// Test 5.1
+	// Parent is fast enough for consideration, child is slower than minimum ping.
+	// Difference in ping is within constant tolerance (i.e. would pass test 4).
+	// Different in ping is GREATER THAN allowed proportional tolerance, as a %age of parent's ping.
+	// Expected result: Pick child. (override fails)
+	{
+		FDatacenterQosInstance ChildSubregion;
+		ChildSubregion.Definition = ChildDatacenter;
+		ChildSubregion.AvgPingMs = 10;
+
+		FDatacenterQosInstance ParentSubregion;
+		ParentSubregion.Definition = ParentDatacenter;
+		ParentSubregion.AvgPingMs = 24; // (24 - 10 = 14) > (0.5 * 24 = 12)
+
+		const FDatacenterQosInstance& RecommendedSubregion = FDatacenterQosInstance::CompareBiasedTowardsNonSubspace(
+			ParentSubregion, ChildSubregion, Invariants);
+		check(&RecommendedSubregion == &ChildSubregion);
+	}
+
+	// Test 5.2
+	// Parent is fast enough for consideration, child is slower than minimum ping.
+	// Difference in ping is within constant tolerance (i.e. would pass test 4).
+	// Different in ping is LESS THAN than allowed proportional tolerance, as a %age of parent's ping.
+	// Expected result: Pick parent. (override succeeds)
+	{
+		FDatacenterQosInstance ChildSubregion;
+		ChildSubregion.Definition = ChildDatacenter;
+		ChildSubregion.AvgPingMs = 10;
+
+		FDatacenterQosInstance ParentSubregion;
+		ParentSubregion.Definition = ParentDatacenter;
+		ParentSubregion.AvgPingMs = 16; // (16 - 10 = 6) < (0.5 * 16 = 8)
+
+		const FDatacenterQosInstance& RecommendedSubregion = FDatacenterQosInstance::CompareBiasedTowardsNonSubspace(
+			ParentSubregion, ChildSubregion, Invariants);
+		check(&RecommendedSubregion == &ParentSubregion);
+	}
+
+	// Test 5.3
+	// Parent is fast enough for consideration, child is slower than minimum ping.
+	// Difference in ping is within constant tolerance (i.e. would pass test 4).
+	// Different in ping is EQUAL TO allowed proportional tolerance, as a %age of parent's ping.
+	// Expected result: Pick parent. (override succeeds)
+	{
+		FDatacenterQosInstance ChildSubregion;
+		ChildSubregion.Definition = ChildDatacenter;
+		ChildSubregion.AvgPingMs = 12;
+
+		FDatacenterQosInstance ParentSubregion;
+		ParentSubregion.Definition = ParentDatacenter;
+		ParentSubregion.AvgPingMs = 24; // (24 - 12 = 12) == (0.5 * 24 = 12)
+
+		const FDatacenterQosInstance& RecommendedSubregion = FDatacenterQosInstance::CompareBiasedTowardsNonSubspace(
+			ParentSubregion, ChildSubregion, Invariants);
+		check(&RecommendedSubregion == &ParentSubregion);
+	}
+
+	return true;
+}
+
+bool UQosRegionManager::TestSortDatacenterSubspacesByRecommended()
+{
+	const FQosSubspaceComparisonParams Invariants(300, 8, 25, 50.0f);
+
+	// Example region
+
+	FQosRegionInfo ExampleRegion;
+	ExampleRegion.DisplayName = NSLOCTEXT("MMRegion", "Europe", "Europe");
+	ExampleRegion.RegionId = TEXT("EU");
+
+	FQosDatacenterInfo ParentDatacenter_DE;
+	ParentDatacenter_DE.Id = TEXT("DE");
+	ParentDatacenter_DE.RegionId = TEXT("EU");
+	ParentDatacenter_DE.Servers.Add(FQosPingServerInfo{ "100.1.1.1", 2345 });
+	ParentDatacenter_DE.Servers.Add(FQosPingServerInfo{ "100.1.1.2", 2345 });
+	ParentDatacenter_DE.Servers.Add(FQosPingServerInfo{ "100.1.1.3", 2345 });
+
+	FQosDatacenterInfo ChildDatacenter_DE_S1;
+	ChildDatacenter_DE_S1.Id = TEXT("DE_S1");
+	ChildDatacenter_DE_S1.RegionId = TEXT("EU");
+	ChildDatacenter_DE_S1.Servers.Add(FQosPingServerInfo{ "100.2.1.11", 2345 });
+	ChildDatacenter_DE_S1.Servers.Add(FQosPingServerInfo{ "100.2.1.12", 2345 });
+	ChildDatacenter_DE_S1.Servers.Add(FQosPingServerInfo{ "100.2.1.13", 2345 });
+
+	FQosDatacenterInfo ChildDatacenter_DE_S2;
+	ChildDatacenter_DE_S2.Id = TEXT("DE_S2");
+	ChildDatacenter_DE_S2.RegionId = TEXT("EU");
+	ChildDatacenter_DE_S2.Servers.Add(FQosPingServerInfo{ "100.3.1.11", 2345 });
+	ChildDatacenter_DE_S2.Servers.Add(FQosPingServerInfo{ "100.3.1.12", 2345 });
+	ChildDatacenter_DE_S2.Servers.Add(FQosPingServerInfo{ "100.3.1.13", 2345 });
+
+	FQosDatacenterInfo ParentDatacenter_FR;
+	ParentDatacenter_FR.Id = TEXT("FR");
+	ParentDatacenter_FR.RegionId = TEXT("EU");
+	ParentDatacenter_FR.Servers.Add(FQosPingServerInfo{ "105.1.1.1", 2345 });
+	ParentDatacenter_FR.Servers.Add(FQosPingServerInfo{ "105.1.1.2", 2345 });
+	ParentDatacenter_FR.Servers.Add(FQosPingServerInfo{ "105.1.1.3", 2345 });
+
+	FQosDatacenterInfo ParentDatacenter_GB;
+	ParentDatacenter_GB.Id = TEXT("GB");
+	ParentDatacenter_GB.RegionId = TEXT("EU");
+	ParentDatacenter_GB.Servers.Add(FQosPingServerInfo{ "120.1.1.1", 2345 });
+	ParentDatacenter_GB.Servers.Add(FQosPingServerInfo{ "120.1.1.2", 2345 });
+	ParentDatacenter_GB.Servers.Add(FQosPingServerInfo{ "120.1.1.3", 2345 });
+
+	FQosDatacenterInfo ChildDatacenter_GB_S;
+	ChildDatacenter_GB_S.Id = TEXT("GB_S");
+	ChildDatacenter_GB_S.RegionId = TEXT("EU");
+	ChildDatacenter_GB_S.Servers.Add(FQosPingServerInfo{ "120.2.1.11", 2345 });
+	ChildDatacenter_GB_S.Servers.Add(FQosPingServerInfo{ "120.2.1.12", 2345 });
+	ChildDatacenter_GB_S.Servers.Add(FQosPingServerInfo{ "120.2.1.13", 2345 });
+
+	const TCHAR* const SubspaceDelimiter = TEXT("_");
+
+	// Test cases.
+
+	FRegionQosInstance TestRegionQosResult;
+	TestRegionQosResult.Definition = ExampleRegion;
+
+	{
+		FDatacenterQosInstance SubregionResult_DE;
+		SubregionResult_DE.Definition = ParentDatacenter_DE;
+		SubregionResult_DE.AvgPingMs = 56;
+
+		FDatacenterQosInstance SubregionResult_DE_S1;
+		SubregionResult_DE_S1.Definition = ChildDatacenter_DE_S1;
+		SubregionResult_DE_S1.AvgPingMs = 48;
+
+		FDatacenterQosInstance SubregionResult_DE_S2;
+		SubregionResult_DE_S2.Definition = ChildDatacenter_DE_S2;
+		SubregionResult_DE_S2.AvgPingMs = 36;
+
+		FDatacenterQosInstance SubregionResult_FR;
+		SubregionResult_FR.Definition = ParentDatacenter_FR;
+		SubregionResult_FR.AvgPingMs = 74;
+
+		FDatacenterQosInstance SubregionResult_GB;
+		SubregionResult_GB.Definition = ParentDatacenter_GB;
+		SubregionResult_GB.AvgPingMs = 92;
+
+		FDatacenterQosInstance SubregionResult_GB_S;
+		SubregionResult_GB_S.Definition = ChildDatacenter_GB_S;
+		SubregionResult_GB_S.AvgPingMs = 84;
+
+		TestRegionQosResult.DatacenterOptions.Empty();
+		TestRegionQosResult.DatacenterOptions.Add(SubregionResult_DE);
+		TestRegionQosResult.DatacenterOptions.Add(SubregionResult_DE_S1);
+		TestRegionQosResult.DatacenterOptions.Add(SubregionResult_DE_S2);
+		TestRegionQosResult.DatacenterOptions.Add(SubregionResult_FR);
+		TestRegionQosResult.DatacenterOptions.Add(SubregionResult_GB);
+		TestRegionQosResult.DatacenterOptions.Add(SubregionResult_GB_S);
+
+		TestRegionQosResult.SortDatacenterSubspacesByRecommended(Invariants, SubspaceDelimiter);
+
+		check(TestRegionQosResult.DatacenterOptions[0].Definition.Id == SubregionResult_DE.Definition.Id); // avg ping = 56
+		check(TestRegionQosResult.DatacenterOptions[1].Definition.Id == SubregionResult_DE_S2.Definition.Id); // avg ping = 36
+		check(TestRegionQosResult.DatacenterOptions[2].Definition.Id == SubregionResult_DE_S1.Definition.Id); // avg ping = 48
+		check(TestRegionQosResult.DatacenterOptions[3].Definition.Id == SubregionResult_FR.Definition.Id); // avg ping = 74
+		check(TestRegionQosResult.DatacenterOptions[4].Definition.Id == SubregionResult_GB.Definition.Id); // avg ping = 92
+		check(TestRegionQosResult.DatacenterOptions[5].Definition.Id == SubregionResult_GB_S.Definition.Id); // avg ping = 84
+	}
+
+	return true;
+}
+
+FRegionQosInstance UQosRegionManager::TestCreateExampleRegionResult()
+{
+	// Definitions.
+
+	static const FString RegionId = TEXT("NAE");
+
+	FQosRegionInfo ExampleRegion;
+	ExampleRegion.DisplayName = NSLOCTEXT("MMRegion", "NA-East", "NA-East");
+	ExampleRegion.RegionId = RegionId;
+	ExampleRegion.RegionId = RegionId;
+	ExampleRegion.bEnabled = true;
+	ExampleRegion.bVisible = true;
+	ExampleRegion.bAutoAssignable = true;
+
+	ExampleRegion.bAllowSubspaceBias = true;
+	ExampleRegion.SubspaceBiasParams = FQosSubspaceComparisonParams(300, 8, 25, 50.0f);
+
+	FQosDatacenterInfo ParentDatacenter_VA;
+	ParentDatacenter_VA.Id = TEXT("VA");
+	ParentDatacenter_VA.RegionId = RegionId;
+	ParentDatacenter_VA.bEnabled = true;
+	ParentDatacenter_VA.Servers.Add(FQosPingServerInfo{ "334.193.154.39", 22222 });
+	ParentDatacenter_VA.Servers.Add(FQosPingServerInfo{ "352.203.3.55", 22222 });
+	ParentDatacenter_VA.Servers.Add(FQosPingServerInfo{ "354.82.195.216", 22222 });
+	ParentDatacenter_VA.Servers.Add(FQosPingServerInfo{ "334.194.116.183", 22222 });
+
+	FQosDatacenterInfo ChildDatacenter_VA_S1;
+	ChildDatacenter_VA_S1.Id = TEXT("VA_S1");
+	ChildDatacenter_VA_S1.RegionId = RegionId;
+	ChildDatacenter_VA_S1.bEnabled = true;
+	ChildDatacenter_VA_S1.Servers.Add(FQosPingServerInfo{ "352.203.34.155", 22222 });
+	ChildDatacenter_VA_S1.Servers.Add(FQosPingServerInfo{ "352.203.34.156", 22222 });
+	ChildDatacenter_VA_S1.Servers.Add(FQosPingServerInfo{ "352.203.34.157", 22222 });
+
+	FQosDatacenterInfo ChildDatacenter_VA_S2;
+	ChildDatacenter_VA_S2.Id = TEXT("VA_S2");
+	ChildDatacenter_VA_S2.RegionId = RegionId;
+	ChildDatacenter_VA_S2.bEnabled = true;
+	ChildDatacenter_VA_S2.Servers.Add(FQosPingServerInfo{ "354.82.199.216", 22222 });
+	ChildDatacenter_VA_S2.Servers.Add(FQosPingServerInfo{ "354.82.199.217", 22222 });
+	ChildDatacenter_VA_S2.Servers.Add(FQosPingServerInfo{ "354.82.199.218", 22222 });
+
+	FQosDatacenterInfo ChildDatacenter_VA_S3;
+	ChildDatacenter_VA_S3.Id = TEXT("VA_S3");
+	ChildDatacenter_VA_S3.RegionId = RegionId;
+	ChildDatacenter_VA_S3.bEnabled = true;
+	ChildDatacenter_VA_S3.Servers.Add(FQosPingServerInfo{ "334.194.111.183", 22222 });
+	ChildDatacenter_VA_S3.Servers.Add(FQosPingServerInfo{ "334.194.111.184", 22222 });
+	ChildDatacenter_VA_S3.Servers.Add(FQosPingServerInfo{ "334.194.111.185", 22222 });
+
+	FQosDatacenterInfo ParentDatacenter_OH;
+	ParentDatacenter_OH.Id = TEXT("OH");
+	ParentDatacenter_OH.RegionId = RegionId;
+	ParentDatacenter_OH.bEnabled = true;
+	ParentDatacenter_OH.Servers.Add(FQosPingServerInfo{ "313.59.18.131", 22222 });
+	ParentDatacenter_OH.Servers.Add(FQosPingServerInfo{ "318.216.47.148", 22222 });
+	ParentDatacenter_OH.Servers.Add(FQosPingServerInfo{ "318.221.198.242", 22222 });
+	ParentDatacenter_OH.Servers.Add(FQosPingServerInfo{ "352.15.144.157", 22222 });
+
+	FQosDatacenterInfo ChildDatacenter_OH_S;
+	ChildDatacenter_OH_S.Id = TEXT("OH_S");
+	ChildDatacenter_OH_S.RegionId = RegionId;
+	ChildDatacenter_OH_S.bEnabled = true;
+	ChildDatacenter_OH_S.Servers.Add(FQosPingServerInfo{ "318.216.49.148", 22222 });
+	ChildDatacenter_OH_S.Servers.Add(FQosPingServerInfo{ "318.216.49.149", 22222 });
+	ChildDatacenter_OH_S.Servers.Add(FQosPingServerInfo{ "318.216.49.150", 22222 });
+
+	// Results
+
+	FRegionQosInstance OutRegionResult;
+	OutRegionResult.Definition = ExampleRegion;
+	OutRegionResult.DatacenterOptions.Empty();
+
+	FDatacenterQosInstance SubregionResult_VA;
+	SubregionResult_VA.Definition = ParentDatacenter_VA;
+	SubregionResult_VA.AvgPingMs = 96;
+	OutRegionResult.DatacenterOptions.Add(SubregionResult_VA);
+
+	FDatacenterQosInstance SubregionResult_VA_S1;
+	SubregionResult_VA_S1.Definition = ChildDatacenter_VA_S1;
+	SubregionResult_VA_S1.AvgPingMs = 91;
+	OutRegionResult.DatacenterOptions.Add(SubregionResult_VA_S1);
+
+	FDatacenterQosInstance SubregionResult_VA_S2;
+	SubregionResult_VA_S2.Definition = ChildDatacenter_VA_S2;
+	SubregionResult_VA_S2.AvgPingMs = 88;
+	OutRegionResult.DatacenterOptions.Add(SubregionResult_VA_S2);
+
+	FDatacenterQosInstance SubregionResult_VA_S3;
+	SubregionResult_VA_S3.Definition = ChildDatacenter_VA_S3;
+	SubregionResult_VA_S3.AvgPingMs = 90;
+	OutRegionResult.DatacenterOptions.Add(SubregionResult_VA_S3);
+
+	FDatacenterQosInstance SubregionResult_OH;
+	SubregionResult_OH.Definition = ParentDatacenter_OH;
+	SubregionResult_OH.AvgPingMs = 117;
+	OutRegionResult.DatacenterOptions.Add(SubregionResult_OH);
+
+	FDatacenterQosInstance SubregionResult_OH_S;
+	SubregionResult_OH_S.Definition = ChildDatacenter_OH_S;
+	SubregionResult_OH_S.AvgPingMs = 97;
+	OutRegionResult.DatacenterOptions.Add(SubregionResult_OH_S);
+
+	return OutRegionResult;
+}
+
+#endif // DEBUG_SUBCOMPARE_BY_SUBSPACE
diff --git a/Engine/Plugins/Online/OnlineFramework/Source/Qos/Public/QosRegionManager.h b/Engine/Plugins/Online/OnlineFramework/Source/Qos/Public/QosRegionManager.h
--- a/Engine/Plugins/Online/OnlineFramework/Source/Qos/Public/QosRegionManager.h
+++ b/Engine/Plugins/Online/OnlineFramework/Source/Qos/Public/QosRegionManager.h
@@ -13,6 +13,8 @@ class UQosEvaluator;
 
 #define UNREACHABLE_PING 9999
 
+#define DEBUG_SUBCOMPARE_BY_SUBSPACE 0
+
 /** Enum for single region QoS return codes */
 UENUM()
 enum class EQosDatacenterResult : uint8
@@ -62,6 +64,72 @@ inline const TCHAR* LexToString(EQosCompletionResult Result)
 	}
 }
 
+/**
+ * Parameters to control the rules-based comparison of subspace vs non-subspace datacenter QoS results.
+ * 
+ * @see FDatacenterQosInstance::IsNonSubspaceRecommended(const FDatacenterQosInstance&, const FDatacenterQosInstance&, const FQosSubspaceComparisonParams&)
+ */
+USTRUCT()
+struct FQosSubspaceComparisonParams
+{
+	GENERATED_USTRUCT_BODY()
+
+	/**
+	 * Maximum allowed ping of the non-subspace.
+	 * If greater than this value, it is too slow, so fails to qualify.
+	 * Set to zero or less to disable checks against this field.
+	 */
+	UPROPERTY()
+	int32 MaxNonSubspacePingMs;
+
+	/**
+	 * Minimum allowed ping of the subspace.
+	 * If below this value, it should not be overridden by the non-subspace.
+	 * Set to zero or less to disable checks against this field.
+	 */
+	UPROPERTY()
+	int32 MinSubspacePingMs;
+
+	/**
+	 * Maximum allowed difference between the subspace and non-subspace's ping values in milliseconds.
+	 * If greater than this value, the non-subspace is too slow, so fails to qualify.
+	 * Set to zero or less to disable checks against this field.
+	 */
+	UPROPERTY()
+	int32 ConstantMaxToleranceMs;
+
+	/**
+	 * Maximum allowed difference between the subspace and non-subspace's ping values,
+	 * which scales as a proportion of the non-subspace's ping, so will differ between
+	 * comparisons when sorting a single list of datacenters.
+	 * If greater than the scaled difference, the non-subspace is too slow, so fails to qualify. 
+	 * Set to zero or less to disable checks against this field.
+	 */
+	UPROPERTY()
+	float ScaledMaxTolerancePct;
+
+	FQosSubspaceComparisonParams()
+		: MaxNonSubspacePingMs(0)
+		, MinSubspacePingMs(0)
+		, ConstantMaxToleranceMs(0)
+		, ScaledMaxTolerancePct(0.0f)
+	{
+	}
+
+	FQosSubspaceComparisonParams(int32 InMaxNonSubspacePingMs, int32 InMinSubspacePingMs, int32 InConstantMaxToleranceMs, float InScaledMaxTolerancePct)
+		: MaxNonSubspacePingMs(InMaxNonSubspacePingMs)
+		, MinSubspacePingMs(InMinSubspacePingMs)
+		, ConstantMaxToleranceMs(InConstantMaxToleranceMs)
+		, ScaledMaxTolerancePct(InScaledMaxTolerancePct)
+	{
+	}
+
+	float CalcScaledMaxToleranceMs(int32 PingMs) const
+	{
+		return 0.01f * ScaledMaxTolerancePct * PingMs;
+	}
+};
+
 /**
  * Individual ping server details
  */
@@ -114,6 +182,11 @@ struct FQosDatacenterInfo
 		return bEnabled && IsValid();
 	}
 
+	bool IsSubspace(const TCHAR* const SubspaceDelimiter) const
+	{
+		return Id.Contains(SubspaceDelimiter, ESearchCase::IgnoreCase, ESearchDir::FromStart);
+	}
+
 	FString ToDebugString() const
 	{
 		return FString::Printf(TEXT("[%s][%s]"), *RegionId, *Id);
@@ -143,11 +216,18 @@ struct FQosRegionInfo
 	/** Can this region be considered for auto detection */
 	UPROPERTY()
 	bool bAutoAssignable;
+	/** Enable biased sorting algorithm on results for this region, which prefers non-subspaces over subspaces */
+	UPROPERTY()
+	bool bAllowSubspaceBias;
+	/** Granular settings for biased subspace-based sorting algorithm, if enabled for this region */
+	UPROPERTY()
+	FQosSubspaceComparisonParams SubspaceBiasParams;
 
 	FQosRegionInfo()
 		: bEnabled(true)
 		, bVisible(true)
 		, bAutoAssignable(true)
+		, bAllowSubspaceBias(false)
 	{
 	}
 
@@ -232,6 +312,67 @@ struct QOS_API FDatacenterQosInstance
 		LastCheckTimestamp = FDateTime(0);
 		bUsable = false;
 	}
+
+	/**
+	 * Compares the avg ping of two datacenters, handling cases where one is a subspace
+	 * and the other is not, and like-for-like.
+	 * 
+	 * When comparing subspace vs non-subspace, this will bias towards the non-subspace,
+	 * as long as it satisfies the series of qualifying rules.
+	 * When comparing like-for-like, average ping is compared, as usual.
+	 * 
+	 * @param A - Left-hand datacenter QoS data to compare
+	 * @param B - Right-hand datacenter QoS data to compare
+	 * @param ComparisonParams - Rules settings for subspace vs non-subspace comparison
+	 * @param SubspaceDelimiter - Search term to look for in datacenter ID; if found, implies that it is a subspace
+	 * @return True if left-hand datacenter is "better", otherwise false (right-hand datacenter is "better")
+	 * 
+	 * @see FDatacenterQosInstance::IsNonSubspaceRecommended(const FDatacenterQosInstance&, const FDatacenterQosInstance&, const FQosSubspaceComparisonParams&)
+	 */
+	static bool IsLessWhenBiasedTowardsNonSubspace(const FDatacenterQosInstance& A, const FDatacenterQosInstance& B,
+		const FQosSubspaceComparisonParams& ComparisonParams, const TCHAR* const SubspaceDelimiter);
+
+	/**
+	 * Compares a subspace datacenter and a non-subspace datacenter, applying the qualifying
+	 * rules to bias non-subspaces, configured via the supplied comparison parameters.
+	 * 
+	 * @param NonSubspace - The non-subspace datacenter QoS data (must be already known to not be a subspace)
+	 * @param Subspace - The subspace datacenter QoS data (must be already known to be a subspace)
+	 * @param ComparisonParams - Granulator adjustments to the comparison rules
+	 * @return True if the NonSubspace is "better" (i.e. passes the qualifying rules), otherwise false.
+	 */
+	static bool IsNonSubspaceRecommended(const FDatacenterQosInstance& NonSubspace, const FDatacenterQosInstance& Subspace,
+		const FQosSubspaceComparisonParams& ComparisonParams);
+
+	/**
+	 * Compares a subspace datacenter and a non-subspace datacenter, applying the qualifying
+	 * rules to bias non-subspaces, configured via the supplied comparison parameters.
+	 *
+	 * @param NonSubspace - The non-subspace datacenter QoS data (must be already known to not be a subspace)
+	 * @param Subspace - The subspace datacenter QoS data (must be already known to be a subspace)
+	 * @param ComparisonParams - Granulator adjustments to the comparison rules
+	 * @return A reference to the input datacenter that is considered "better" after rules-based comparison.
+	 *
+	 * @see FDatacenterQosInstance::IsNonSubspaceRecommended(const FDatacenterQosInstance&, const FDatacenterQosInstance&, const FQosSubspaceComparisonParams&)
+	 */
+	static const FDatacenterQosInstance& CompareBiasedTowardsNonSubspace(
+		const FDatacenterQosInstance& NonSubspace, const FDatacenterQosInstance& Subspace,
+		const FQosSubspaceComparisonParams& ComparisonParams);
+
+	/**
+	 * Compares a subspace datacenter and a non-subspace datacenter, applying the qualifying
+	 * rules to bias non-subspaces, configured via the supplied comparison parameters.
+	 *
+	 * @param NonSubspace - The non-subspace datacenter QoS data (must be already known to not be a subspace)
+	 * @param Subspace - The subspace datacenter QoS data (must be already known to be a subspace)
+	 * @param ComparisonParams - Granulator adjustments to the comparison rules
+	 * @return A pointer to the input datacenter that is considered "better" after rules-based comparison.
+	 *
+	 * @see FDatacenterQosInstance::IsNonSubspaceRecommended(const FDatacenterQosInstance&, const FDatacenterQosInstance&, const FQosSubspaceComparisonParams&)
+	 */
+	static const FDatacenterQosInstance* const CompareBiasedTowardsNonSubspace(
+		const FDatacenterQosInstance* const NonSubspace, const FDatacenterQosInstance* const Subspace,
+		const FQosSubspaceComparisonParams& ComparisonParams);
 };
 
 USTRUCT()
@@ -277,12 +418,31 @@ struct QOS_API FRegionQosInstance
 
 	/** @return the result of this region ping request */
 	EQosDatacenterResult GetRegionResult() const;
+
 	/** @return the ping recorded in the best sub region */
 	int32 GetBestAvgPing() const;
+
 	/** @return the subregion with the best ping */
 	FString GetBestSubregion() const;
+
 	/** @return sorted list of subregions by best ping */
 	void GetSubregionPreferences(TArray<FString>& OutSubregions) const;
+
+	/** Sort the list of datacenter options into ascending order of average ping */
+	void SortDatacenterOptionsByAvgPingAsc();
+
+	/**
+	 * Sort the list of datacenter options into ascending order, using rules-based comparison
+	 * for cases where a subspace is being compared to a non-subspace; non-subspace will be
+	 * favoured if it passes the rules check.
+	 * Like-for-like comparisons are favour lower average ping.
+	 * 
+	 * @param ComparisonParams - Granulator adjustments to the comparison rules
+	 * @param SubspaceDelimiter - Search term to look for in datacenter ID; if found, implies that it is a subspace
+	 *
+	 * @see FDatacenterQosInstance::IsLessWhenBiasedTowardsNonSubspace(const FDatacenterQosInstance&, const FDatacenterQosInstance&, const FQosSubspaceComparisonParams&, const TCHAR*);
+	 */
+	void SortDatacenterSubspacesByRecommended(const FQosSubspaceComparisonParams& ComparisonParams, const TCHAR* const SubspaceDelimiter);
 };
 
 /**
@@ -398,6 +558,36 @@ public:
 	virtual void PostReloadConfig(class FProperty* PropertyThatWasLoaded) override;
 	/** End UObject interface */
 
+private:
+
+	/**
+	 * Get the delimiter string used to split primary subspace ID from a subregion ID.
+	 */
+	const TCHAR* GetSubspaceDelimiter() const;
+
+	/**
+	 * Finds the QOS region result that matches the given region ID from an array of region results.
+	 * 
+	 * @return pointer the the region result if found, otherwise nullptr
+	 */
+	static const FRegionQosInstance* FindQosRegionById(const TArray<FRegionQosInstance>& Regions, const FString& RegionId);
+
+	/**
+	 * Finds the QOS region result that has the "best" ping result.
+	 * Assumes that the datacenter results within each region result are pre-sorted,
+	 * e.g. by average ping, or via a recommendation bias algorithm.
+	 * 
+	 * @return pointer to the "best" region result, determined by a previous sort of region datacenters, otherwise nullptr if no results exist.
+	 */
+	static const FRegionQosInstance* FindBestQosRegion(const TArray<FRegionQosInstance>& Regions);
+
+#ifdef DEBUG_SUBCOMPARE_BY_SUBSPACE
+	// Test methods for debugging datacenter comparisons that use special rules when dealing with subspaces.
+	static bool TestCompareDatacentersBySubspace();
+	static bool TestSortDatacenterSubspacesByRecommended();
+	static FRegionQosInstance TestCreateExampleRegionResult();
+#endif // DEBUG_SUBCOMPARE_BY_SUBSPACE
+
 private:
 
 	/**
@@ -417,6 +607,28 @@ private:
 	 */
 	int32 GetMaxPingMs() const;
 
+	/**
+	 * Should datacenter QoS results be sorted using rules-based comparison where subspaces are encountered?
+	 * 
+	 * Determined via bEnableSubspaceBiasOrder engine configuration param for QosRegionManager.
+	 * Global enable/disable may be overridden by `qossubspacebias=true|false` command-line arg.
+	 * 
+	 * @return True if rules-based sorting (when dealing with subspaces) is enabled, otherwise false.
+	 */
+	bool IsSubspaceBiasOrderEnabled() const;
+
+	/**
+	 * Should datacenter QoS results be sorted using rules-based comparison where subspaces are encountered
+	 * for the given region's QoS data?
+	 * 
+	 * Determined via bEnableSubspaceBiasOrder engine configuration param for QosRegionManager,
+	 * and the specific region's RegionDefinition entry.
+	 * Global enable/disable may be overridden by `qossubspacebias=true|false` command-line arg.
+	 *
+	 * @return True if rules-based sorting (when dealing with subspaces) is enabled, otherwise false.
+	 */
+	bool IsSubspaceBiasOrderEnabled(const FQosRegionInfo& RegionDefinition) const;
+
 	/** Number of times to ping a given region using random sampling of available servers */
 	UPROPERTY(Config)
 	int32 NumTestsPerRegion;
@@ -425,6 +637,20 @@ private:
 	UPROPERTY(Config)
 	float PingTimeout;
 
+	/**
+	 * Global switch to enable/disable sorting of QoS datacenter results using rules-based comparison,
+	 * where subspaces are encountered.
+	 */
+	UPROPERTY(Config)
+	bool bEnableSubspaceBiasOrder;
+
+	/**
+	 * Delimiter string that identifies a subspace datacenter ID.
+	 * e.g. "DE_S" would be a subspace of the "DE" subregion, using "_" as the delimiter.
+	 */
+	UPROPERTY(Config)
+	FString SubspaceDelimiter;
+
 	/** Metadata about existing regions */
 	UPROPERTY(Config)
 	TArray<FQosRegionInfo> RegionDefinitions;
@@ -459,5 +685,7 @@ private:
 	FOnQosEvalCompleteDelegate OnQosEvalCompleteDelegate;
 
 	FSimpleDelegate OnQoSSettingsChangedDelegate;
+
+	static const TCHAR* SubspaceDelimiterDefault;
 };
 
