SocialAnalytics, convert social user analytics context to TMap to allow arbitrary number of contextual analytics params

When we added the SocialMenuContext analytic we added the ability to execute a social interaction "with context". At the time the context was implemented as a string. We now have a request for more functionality here, so converting the string to a dictionary to allow for more params.

#tests
launched client with -AnalyticsDisableCaching unfriended and re-friended someone, checked the analytics string in the logs for the proper values

diff --git a/Engine/Plugins/Online/OnlineFramework/Source/Party/Private/Interactions/SocialInteractionHandle.cpp b/Engine/Plugins/Online/OnlineFramework/Source/Party/Private/Interactions/SocialInteractionHandle.cpp
--- a/Engine/Plugins/Online/OnlineFramework/Source/Party/Private/Interactions/SocialInteractionHandle.cpp
+++ b/Engine/Plugins/Online/OnlineFramework/Source/Party/Private/Interactions/SocialInteractionHandle.cpp
@@ -45,7 +45,7 @@ void FSocialInteractionHandle::ExecuteInteraction(USocialUser& User) const
 	}
 }
 
-void FSocialInteractionHandle::ExecuteInteractionWithContext(USocialUser& User, const FString& AnalyticsContext) const
+void FSocialInteractionHandle::ExecuteInteractionWithContext(USocialUser& User, const TMap<FString, FString>& AnalyticsContext) const
 {
 	if (InteractionWrapper)
 	{
diff --git a/Engine/Plugins/Online/OnlineFramework/Source/Party/Private/User/SocialUser.cpp b/Engine/Plugins/Online/OnlineFramework/Source/Party/Private/User/SocialUser.cpp
--- a/Engine/Plugins/Online/OnlineFramework/Source/Party/Private/User/SocialUser.cpp
+++ b/Engine/Plugins/Online/OnlineFramework/Source/Party/Private/User/SocialUser.cpp
@@ -994,12 +994,12 @@ FString USocialUser::ToDebugString() const
 #endif
 }
 
-void USocialUser::WithContext(const FString& InAnalyticsContext, void(*Func)(USocialUser&))
+void USocialUser::WithContext(const TMap<FString, FString>& InAnalyticsContext, void(*Func)(USocialUser&))
 {
-	const FString PreviousContext = AnalyticsContext;
+	TMap<FString, FString> PreviousContext = MoveTemp(AnalyticsContext);
 	AnalyticsContext = InAnalyticsContext;
 	Func(*this);
-	AnalyticsContext = PreviousContext;
+	AnalyticsContext = MoveTemp(PreviousContext);
 }
 
 bool USocialUser::SendFriendInvite(ESocialSubsystem SubsystemType)
diff --git a/Engine/Plugins/Online/OnlineFramework/Source/Party/Public/Interactions/SocialInteractionHandle.h b/Engine/Plugins/Online/OnlineFramework/Source/Party/Public/Interactions/SocialInteractionHandle.h
--- a/Engine/Plugins/Online/OnlineFramework/Source/Party/Public/Interactions/SocialInteractionHandle.h
+++ b/Engine/Plugins/Online/OnlineFramework/Source/Party/Public/Interactions/SocialInteractionHandle.h
@@ -26,7 +26,7 @@ public:
 
 	bool IsAvailable(const USocialUser& User) const;
 	void ExecuteInteraction(USocialUser& User) const;
-	void ExecuteInteractionWithContext(USocialUser& User, const FString& AnalyticsContext) const;
+	void ExecuteInteractionWithContext(USocialUser& User, const TMap<FString, FString>& AnalyticsContext) const;
 
 private:
 	template <typename> friend class TSocialInteractionWrapper;
diff --git a/Engine/Plugins/Online/OnlineFramework/Source/Party/Public/Interactions/SocialInteractionMacros.h b/Engine/Plugins/Online/OnlineFramework/Source/Party/Public/Interactions/SocialInteractionMacros.h
--- a/Engine/Plugins/Online/OnlineFramework/Source/Party/Public/Interactions/SocialInteractionMacros.h
+++ b/Engine/Plugins/Online/OnlineFramework/Source/Party/Public/Interactions/SocialInteractionMacros.h
@@ -21,7 +21,7 @@ public:
 
 	virtual bool IsAvailable(const USocialUser& User) const = 0;
 	virtual void ExecuteInteraction(USocialUser& User) const = 0;
-	virtual void ExecuteInteractionWithContext(USocialUser& User, const FString& AnalyticsContext) const = 0;
+	virtual void ExecuteInteractionWithContext(USocialUser& User, const TMap<FString, FString>& AnalyticsContext) const = 0;
 };
 
 template <typename InteractionT>
@@ -34,7 +34,7 @@ public:
 
 	virtual bool IsAvailable(const USocialUser& User) const override final { return InteractionT::IsAvailable(User); }
 	virtual void ExecuteInteraction(USocialUser& User) const override final { InteractionT::ExecuteInteraction(User); }
-	virtual void ExecuteInteractionWithContext(USocialUser& User, const FString& AnalyticsContext) const override final { InteractionT::ExecuteInteractionWithContext(User, AnalyticsContext); };
+	virtual void ExecuteInteractionWithContext(USocialUser& User, const TMap<FString, FString>& AnalyticsContext) const override final { InteractionT::ExecuteInteractionWithContext(User, AnalyticsContext); };
 
 private:
 	friend InteractionT;
@@ -72,7 +72,7 @@ public:	\
 		return false;	\
 	}	\
 	\
-	static void ExecuteInteractionWithContext(USocialUser& User, const FString& AnalyticsContext)	\
+	static void ExecuteInteractionWithContext(USocialUser& User, const TMap<FString, FString>& AnalyticsContext)	\
 	{	\
 		User.WithContext(AnalyticsContext, [](USocialUser& InUser) {	\
 			ExecuteInteraction(InUser);	\
diff --git a/Engine/Plugins/Online/OnlineFramework/Source/Party/Public/User/SocialUser.h b/Engine/Plugins/Online/OnlineFramework/Source/Party/Public/User/SocialUser.h
--- a/Engine/Plugins/Online/OnlineFramework/Source/Party/Public/User/SocialUser.h
+++ b/Engine/Plugins/Online/OnlineFramework/Source/Party/Public/User/SocialUser.h
@@ -98,8 +98,8 @@ public:
 	virtual bool RejectFriendInvite(ESocialSubsystem SocialSubsystem) const;
 	virtual bool EndFriendship(ESocialSubsystem SocialSubsystem) const;
 
-	FString GetAnalyticsContext() const { return AnalyticsContext; }
-	void WithContext(const FString& InAnalyticsContext, void(*Func)(USocialUser&));
+	TMap<FString, FString> GetAnalyticsContext() const { return AnalyticsContext; }
+	void WithContext(const TMap<FString, FString>& InAnalyticsContext, void(*Func)(USocialUser&));
 
 	bool ShowPlatformProfile();
 
@@ -189,7 +189,7 @@ protected:
 	virtual void NotifyRequestToJoinReceived(const IOnlinePartyRequestToJoinInfo& Request) {}
 	virtual void NotifyRequestToJoinRemoved(const IOnlinePartyRequestToJoinInfo& Request, EPartyRequestToJoinRemovedReason Reason) {}
 	int32 NumPendingQueries = 0;
-	FString AnalyticsContext;
+	TMap<FString, FString> AnalyticsContext;
 
 	void TryBroadcastInitializationComplete();
 
