Additional VeryVerbose logging for QoS results sorting.

diff --git a/Engine/Plugins/Online/OnlineFramework/Source/Qos/Private/QosRegionManager.cpp b/Engine/Plugins/Online/OnlineFramework/Source/Qos/Private/QosRegionManager.cpp
--- a/Engine/Plugins/Online/OnlineFramework/Source/Qos/Private/QosRegionManager.cpp
+++ b/Engine/Plugins/Online/OnlineFramework/Source/Qos/Private/QosRegionManager.cpp
@@ -133,6 +133,8 @@ void FRegionQosInstance::GetSubregionPreferences(TArray<FString>& OutSubregions)
 
 void FRegionQosInstance::SortDatacenterOptionsByAvgPingAsc()
 {
+	UE_LOG(LogQos, VeryVerbose, TEXT("[FRegionQosInstance::SortDatacenterOptionsByAvgPingAsc] sorting datacenter results for region %s"), *GetRegionId());
+
 	// Sort ping best to worst (ascending ping ms)
 	DatacenterOptions.Sort(
 		[](const FDatacenterQosInstance& A, const FDatacenterQosInstance& B) {
@@ -143,6 +145,8 @@ void FRegionQosInstance::SortDatacenterOptionsByAvgPingAsc()
 
 void FRegionQosInstance::SortDatacenterSubspacesByRecommended(const FQosSubspaceComparisonParams& ComparisonParams, const TCHAR* const SubspaceDelimiter)
 {
+	UE_LOG(LogQos, VeryVerbose, TEXT("[FRegionQosInstance::SortDatacenterSubspacesByRecommended] sorting datacenter results for region %s"), *GetRegionId());
+
 	// Sort by subspace recommendation rules
 	DatacenterOptions.Sort(
 		[&ComparisonParams, &SubspaceDelimiter](const FDatacenterQosInstance& A, const FDatacenterQosInstance& B) {
@@ -151,6 +155,17 @@ void FRegionQosInstance::SortDatacenterSubspacesByRecommended(const FQosSubspace
 	);
 }
 
+void FRegionQosInstance::LogDatacenterResults() const
+{
+	UE_LOG(LogQos, VeryVerbose, TEXT("Region %s:"), *GetRegionId());
+
+	for (const FDatacenterQosInstance& Datacenter : DatacenterOptions)
+	{
+		UE_LOG(LogQos, VeryVerbose, TEXT("  %s  avg ping: %dms, num responses: %d"), *Datacenter.Definition.Id, Datacenter.AvgPingMs, Datacenter.NumResponses);
+	}
+}
+
+
 const TCHAR* UQosRegionManager::SubspaceDelimiterDefault = TEXT("_");
 
 UQosRegionManager::UQosRegionManager(const FObjectInitializer& ObjectInitializer)
@@ -253,9 +268,12 @@ bool UQosRegionManager::IsSubspaceBiasOrderEnabled() const
 	bool bEnabledCommandLine = false;
 	if (FParse::Bool(FCommandLine::Get(), TEXT("qossubspacebias="), bEnabledCommandLine))
 	{
+		UE_LOG(LogQos, VeryVerbose, TEXT("[UQosRegionManager::IsSubspaceBiasOrderEnabled] overridden by command-line to: %s"), *LexToString(bEnabledCommandLine));
 		return bEnabledCommandLine;
 	}
 
+	UE_LOG(LogQos, VeryVerbose, TEXT("[UQosRegionManager::IsSubspaceBiasOrderEnabled] set in ini or defaulted to: %s"), *LexToString(bEnableSubspaceBiasOrder));
+
 	// Use config value (or internal code default)
 	return bEnableSubspaceBiasOrder;
 }
@@ -398,6 +416,8 @@ void UQosRegionManager::OnQosEvaluationComplete(EQosCompletionResult Result, con
 	const bool bSubspaceBiasEnabled = IsSubspaceBiasOrderEnabled();
 	const TCHAR* const Delim = GetSubspaceDelimiter();
 
+	UE_LOG(LogQos, Verbose, TEXT("[UQosRegionManager::OnQosEvaluationComplete] bSubspaceBiasEnabled=%s"), *LexToString(bSubspaceBiasEnabled));
+
 #if DEBUG_SUBCOMPARE_BY_SUBSPACE
 	// Test example regions (indepedent of real obtained QoS data)
 	{
@@ -420,6 +440,11 @@ void UQosRegionManager::OnQosEvaluationComplete(EQosCompletionResult Result, con
 				DatacenterMap.MultiFind(RegionInfo.RegionId, NewRegion->DatacenterOptions);
 
 				const bool bSubspaceBias = bSubspaceBiasEnabled && RegionInfo.bAllowSubspaceBias;
+
+				UE_LOG(LogQos, VeryVerbose, TEXT("[UQosRegionManager::OnQosEvaluationComplete] pre-sort results for region=%s, bias=%s"),
+					*RegionInfo.RegionId, *LexToString(bSubspaceBias));
+				NewRegion->LogDatacenterResults();
+
 				if (bSubspaceBias)
 				{
 					// Use ordering that applies rules-based comparison for subspace vs non-subspace,
@@ -432,6 +457,10 @@ void UQosRegionManager::OnQosEvaluationComplete(EQosCompletionResult Result, con
 					// Use regular sorting of subregions, by ascendering order of avg ping.
 					NewRegion->SortDatacenterOptionsByAvgPingAsc();
 				}
+
+				UE_LOG(LogQos, VeryVerbose, TEXT("[UQosRegionManager::OnQosEvaluationComplete] post-sort results for region=%s, bias=%s"),
+					*RegionInfo.RegionId, *LexToString(bSubspaceBias));
+				NewRegion->LogDatacenterResults();
 			}
 			else
 			{
diff --git a/Engine/Plugins/Online/OnlineFramework/Source/Qos/Public/QosRegionManager.h b/Engine/Plugins/Online/OnlineFramework/Source/Qos/Public/QosRegionManager.h
--- a/Engine/Plugins/Online/OnlineFramework/Source/Qos/Public/QosRegionManager.h
+++ b/Engine/Plugins/Online/OnlineFramework/Source/Qos/Public/QosRegionManager.h
@@ -443,6 +443,9 @@ struct QOS_API FRegionQosInstance
 	 * @see FDatacenterQosInstance::IsLessWhenBiasedTowardsNonSubspace(const FDatacenterQosInstance&, const FDatacenterQosInstance&, const FQosSubspaceComparisonParams&, const TCHAR*);
 	 */
 	void SortDatacenterSubspacesByRecommended(const FQosSubspaceComparisonParams& ComparisonParams, const TCHAR* const SubspaceDelimiter);
+
+	/** Print list of datacenters results for this region to QoS log. */
+	void LogDatacenterResults() const;
 };
 
 /**
