Expose IsCompatibleHotfixFile to manager subclasses.

diff --git a/Engine/Plugins/Online/OnlineFramework/Source/Hotfix/Private/OnlineHotfixManager.cpp b/Engine/Plugins/Online/OnlineFramework/Source/Hotfix/Private/OnlineHotfixManager.cpp
--- a/Engine/Plugins/Online/OnlineFramework/Source/Hotfix/Private/OnlineHotfixManager.cpp
+++ b/Engine/Plugins/Online/OnlineFramework/Source/Hotfix/Private/OnlineHotfixManager.cpp
@@ -119,49 +119,39 @@ namespace
 			}
 		}
 	}
+}
 
-	/**
-	 * Is this hotfix file compatible with the current build
-	 * If the file has version information it is compared with compatibility
-	 * If the file has NO version information it is assumed compatible
-	 *
-	 * @param InFilename name of the file to check
-	 * @param OutFilename name of file with version information stripped
-	 *
-	 * @return true if file is compatible, false otherwise
-	 */
-	bool IsCompatibleHotfixFile(const FString& InFilename, FString& OutFilename)
-	{
-		bool bHasNetVersion = false;
-		bool bCompatibleNetHotfix = false;
-		bool bHasBranchVersion = false;
-		bool bCompatibleBranchHotfix = false;
-		FString OutNetVersion;
-		FString OutBranchVersion;
-		GetFilenameAndVersion(InFilename, OutFilename, OutNetVersion, OutBranchVersion);
-
-		if (!OutNetVersion.IsEmpty())
+bool UOnlineHotfixManager::IsCompatibleHotfixFile(const FString& InFilename, FString& OutFilename)
+{
+	bool bHasNetVersion = false;
+	bool bCompatibleNetHotfix = false;
+	bool bHasBranchVersion = false;
+	bool bCompatibleBranchHotfix = false;
+	FString OutNetVersion;
+	FString OutBranchVersion;
+	GetFilenameAndVersion(InFilename, OutFilename, OutNetVersion, OutBranchVersion);
+
+	if (!OutNetVersion.IsEmpty())
+	{
+		bHasNetVersion = true;
+		const FString NetworkVersion = GetNetworkVersion();
+		if (OutNetVersion == NetworkVersion)
 		{
-			bHasNetVersion = true;
-			const FString NetworkVersion = GetNetworkVersion();
-			if (OutNetVersion == NetworkVersion)
-			{
-				bCompatibleNetHotfix = true;
-			}
+			bCompatibleNetHotfix = true;
 		}
+	}
 
-		if (!OutBranchVersion.IsEmpty())
+	if (!OutBranchVersion.IsEmpty())
+	{
+		bHasBranchVersion = true;
+		const FString BranchVersion = GetBranchVersion();
+		if (OutBranchVersion == BranchVersion)
 		{
-			bHasBranchVersion = true;
-			const FString BranchVersion = GetBranchVersion();
-			if (OutBranchVersion == BranchVersion)
-			{
-				bCompatibleBranchHotfix = true;
-			}
+			bCompatibleBranchHotfix = true;
 		}
-
-		return (bCompatibleNetHotfix || !bHasNetVersion) && (bCompatibleBranchHotfix || !bHasBranchVersion);
 	}
+
+	return (bCompatibleNetHotfix || !bHasNetVersion) && (bCompatibleBranchHotfix || !bHasBranchVersion);
 }
 
 UOnlineHotfixManager::UOnlineHotfixManager() :
diff --git a/Engine/Plugins/Online/OnlineFramework/Source/Hotfix/Public/OnlineHotfixManager.h b/Engine/Plugins/Online/OnlineFramework/Source/Hotfix/Public/OnlineHotfixManager.h
--- a/Engine/Plugins/Online/OnlineFramework/Source/Hotfix/Public/OnlineHotfixManager.h
+++ b/Engine/Plugins/Online/OnlineFramework/Source/Hotfix/Public/OnlineHotfixManager.h
@@ -222,6 +222,19 @@ protected:
 	UWorld* GetWorld() const override;
 
 protected:
+
+	/**
+	 * Is this hotfix file compatible with the current build
+	 * If the file has version information it is compared with compatibility
+	 * If the file has NO version information it is assumed compatible
+	 *
+	 * @param InFilename name of the file to check
+	 * @param OutFilename name of file with version information stripped
+	 *
+	 * @return true if file is compatible, false otherwise
+	 */
+	bool IsCompatibleHotfixFile(const FString& InFilename, FString& OutFilename);
+
 	/**
 	 * Override this method to look at the file information for any game specific hotfix processing
 	 * NOTE: Make sure to call Super to get default handling of files
