Fixed an issue where hotfixing would not run in the editor client.
Changed the UpdateManager to use FTicker rather than the timer manager which doesn't always operate in the editor without a ticking world.

Added new function:
virtual bool ShouldPerformHotfix()
so applications can choose if they want to perform hotfixes

Re-enable hotfixing
Backout of CL: 15629769

diff --git a/Engine/Plugins/Online/OnlineFramework/Source/Hotfix/Private/OnlineHotfixManager.cpp b/Engine/Plugins/Online/OnlineFramework/Source/Hotfix/Private/OnlineHotfixManager.cpp
--- a/Engine/Plugins/Online/OnlineFramework/Source/Hotfix/Private/OnlineHotfixManager.cpp
+++ b/Engine/Plugins/Online/OnlineFramework/Source/Hotfix/Private/OnlineHotfixManager.cpp
@@ -271,7 +271,7 @@ void UOnlineHotfixManager::StartHotfixProcess()
 	UE_LOG(LogHotfixManager, Log, TEXT("Starting Hotfix Process"));
 
 	// Patching the editor this way seems like a bad idea
-	bool bShouldHotfix = IsRunningGame() || IsRunningDedicatedServer() || IsRunningClientOnly();
+	const bool bShouldHotfix = ShouldPerformHotfix();
 	if (!bShouldHotfix)
 	{
 		UE_LOG(LogHotfixManager, Warning, TEXT("Hotfixing skipped when not running game/server"));
@@ -460,7 +460,7 @@ void UOnlineHotfixManager::OnEnumerateFilesComplete(bool bWasSuccessful, const F
 void UOnlineHotfixManager::CheckAvailability(FOnHotfixAvailableComplete& InCompletionDelegate)
 {
 	// Checking for hotfixes in editor is not supported
-	bool bShouldHotfix = IsRunningGame() || IsRunningDedicatedServer() || IsRunningClientOnly();
+	const bool bShouldHotfix = ShouldPerformHotfix();
 	if (!bShouldHotfix)
 	{
 		UE_LOG(LogHotfixManager, Warning, TEXT("Hotfixing availability skipped when not running game/server"));
@@ -1939,6 +1939,11 @@ void UOnlineHotfixManager::HotfixTableUpdate(UObject* Asset, const FString& Asse
 	}
 }
 
+bool UOnlineHotfixManager::ShouldPerformHotfix()
+{
+	return IsRunningGame() || IsRunningDedicatedServer() || IsRunningClientOnly();
+}
+
 UWorld* UOnlineHotfixManager::GetWorld() const
 {
 	return OwnerWorld.IsValid() ? OwnerWorld.Get() : nullptr;
diff --git a/Engine/Plugins/Online/OnlineFramework/Source/Hotfix/Private/UpdateManager.cpp b/Engine/Plugins/Online/OnlineFramework/Source/Hotfix/Private/UpdateManager.cpp
--- a/Engine/Plugins/Online/OnlineFramework/Source/Hotfix/Private/UpdateManager.cpp
+++ b/Engine/Plugins/Online/OnlineFramework/Source/Hotfix/Private/UpdateManager.cpp
@@ -121,7 +121,7 @@ UUpdateManager::EUpdateStartResult UUpdateManager::StartCheckInternal(bool bInCh
 		return Result;
 	}
 
-	if (!IsTimerHandleActive(StartCheckInternalTimerHandle) &&
+	if (!StartCheckInternalTimerHandle.IsValid() &&
 		(CurrentUpdateState == EUpdateState::UpdateIdle ||
 		CurrentUpdateState == EUpdateState::UpdatePending ||
 		CurrentUpdateState == EUpdateState::UpdateComplete))
@@ -147,6 +147,7 @@ UUpdateManager::EUpdateStartResult UUpdateManager::StartCheckInternal(bool bInCh
 			{
 				// Check for a patch first, then hotfix application
 				StartPatchCheck();
+				StartCheckInternalTimerHandle.Reset();
 			};
 
 			// Give the UI state widget a chance to start listening for delegates
@@ -159,6 +160,7 @@ UUpdateManager::EUpdateStartResult UUpdateManager::StartCheckInternal(bool bInCh
 			auto StartDelegate = [this, LastResult]()
 			{
 				CheckComplete(LastResult, false);
+				StartCheckInternalTimerHandle.Reset();
 			};
 
 			StartCheckInternalTimerHandle = DelayResponse(StartDelegate, 0.1f);
@@ -642,28 +644,13 @@ void UUpdateManager::OnApplicationHasReactivated()
 	}
 }
 
-FTimerHandle UUpdateManager::DelayResponse(DelayCb&& Delegate, float Delay)
+FTSTicker::FDelegateHandle UUpdateManager::DelayResponse(DelayCb&& Delegate, float Delay)
 {
-	FTimerHandle TimerHandle;
-
-	UWorld* World = GetWorld();
-	if (ensure(World != nullptr))
-	{
-		World->GetTimerManager().SetTimer(TimerHandle, MoveTemp(Delegate), Delay, false, -1.f);
-	}
-
-	return TimerHandle;
-}
-
-bool UUpdateManager::IsTimerHandleActive(const FTimerHandle& TimerHandle) const
-{
-	UWorld* World = GetWorld();
-	if (ensure(World != nullptr))
-	{
-		return World->GetTimerManager().IsTimerActive(TimerHandle);
-	}
-
-	return false;
+	return FTSTicker::GetCoreTicker().AddTicker(FTickerDelegate::CreateWeakLambda(this, [this, Delegate](float dts)
+		{
+			Delegate();
+			return false;
+		}), Delay);
 }
 
 UWorld* UUpdateManager::GetWorld() const
@@ -701,4 +688,4 @@ FString LexToString(EUpdateCompletionStatus Status)
 		default:
 			return TEXT("UpdateUnknown");
 	}
-}
+}
\ No newline at end of file
diff --git a/Engine/Plugins/Online/OnlineFramework/Source/Hotfix/Public/OnlineHotfixManager.h b/Engine/Plugins/Online/OnlineFramework/Source/Hotfix/Public/OnlineHotfixManager.h
--- a/Engine/Plugins/Online/OnlineFramework/Source/Hotfix/Public/OnlineHotfixManager.h
+++ b/Engine/Plugins/Online/OnlineFramework/Source/Hotfix/Public/OnlineHotfixManager.h
@@ -329,6 +329,8 @@ protected:
 	virtual void OnHotfixTableValueObject(UObject& Asset, const FString& RowName, const FString& ColumnName, const UObject* OldValue, const UObject* NewValue) { }
 	virtual void OnHotfixTableValueSoftObject(UObject& Asset, const FString& RowName, const FString& ColumnName, const FSoftObjectPtr& OldValue, const FSoftObjectPtr& NewValue) { }
 
+	virtual bool ShouldPerformHotfix();
+
 public:
 	UOnlineHotfixManager();
 	UOnlineHotfixManager(FVTableHelper& Helper);
diff --git a/Engine/Plugins/Online/OnlineFramework/Source/Hotfix/Public/UpdateManager.h b/Engine/Plugins/Online/OnlineFramework/Source/Hotfix/Public/UpdateManager.h
--- a/Engine/Plugins/Online/OnlineFramework/Source/Hotfix/Public/UpdateManager.h
+++ b/Engine/Plugins/Online/OnlineFramework/Source/Hotfix/Public/UpdateManager.h
@@ -389,7 +389,7 @@ protected:
 	EUpdateCompletionStatus LastCompletionResult[2];
 
 	FTSTicker::FDelegateHandle TickerHandle;
-	FTimerHandle StartCheckInternalTimerHandle;
+	FTSTicker::FDelegateHandle StartCheckInternalTimerHandle;
 
 private:
 
@@ -405,10 +405,7 @@ private:
 
 	/** Fire a delegate after a given amount of time */
 	typedef TFunction<void(void)> DelayCb;
-	FTimerHandle DelayResponse(DelayCb&& Delegate, float Delay);
-
-	/** Helper function to check if a timer is running. */
-	bool IsTimerHandleActive(const FTimerHandle& TimerHandle) const;
+	FTSTicker::FDelegateHandle DelayResponse(DelayCb&& Delegate, float Delay);
 
 	friend bool SkipPatchCheck(UUpdateManager* UpdateManager);
 
