Set so when filtering out users to query UserInfo for, we never filter out the local user to assure freshest information upon logging out and logging back in.

#tests ran PIE and logged in before linking a console account, after being forced to logout by logging in on a console and linking account, then relogging in on PIE, system properly requeried for local user info on login when these subsystems were setup.

diff --git a/Engine/Plugins/Online/OnlineFramework/Source/Party/Private/User/SocialUser.cpp b/Engine/Plugins/Online/OnlineFramework/Source/Party/Private/User/SocialUser.cpp
--- a/Engine/Plugins/Online/OnlineFramework/Source/Party/Private/User/SocialUser.cpp
+++ b/Engine/Plugins/Online/OnlineFramework/Source/Party/Private/User/SocialUser.cpp
@@ -1316,7 +1316,8 @@ void USocialUser::SetSubsystemId(ESocialSubsystem SubsystemType, const FUniqueNe
 			if (UserInterface.IsValid())
 			{
 				TSharedPtr<FOnlineUser> UserInfo = UserInterface->GetUserInfo(OwningToolkit.GetLocalUserNum(), *SubsystemId);
-				if (UserInfo.IsValid())
+				//If this is our local user, we always want to get the freshest data, to make sure any updates to the account (linking, name changes, etc.) are seen on re-login
+				if (UserInfo.IsValid() && !IsLocalUser())
 				{
 					SetUserInfo(SubsystemType, UserInfo.ToSharedRef());
 				}
