[QoS] add an OnQosRegionIdChanged delegate to the QoS interface

diff --git a/Engine/Plugins/Online/OnlineFramework/Source/Qos/Private/QosInterface.cpp b/Engine/Plugins/Online/OnlineFramework/Source/Qos/Private/QosInterface.cpp
--- a/Engine/Plugins/Online/OnlineFramework/Source/Qos/Private/QosInterface.cpp
+++ b/Engine/Plugins/Online/OnlineFramework/Source/Qos/Private/QosInterface.cpp
@@ -121,4 +121,10 @@ void FQosInterface::RegisterQoSSettingsChangedDelegate(const FSimpleDelegate& On
 {
 	check(RegionManager);
 	RegionManager->RegisterQoSSettingsChangedDelegate(OnQoSSettingsChanged);
+}
+
+FQosInterface::FOnQosRegionIdChanged& FQosInterface::OnQosRegionIdChanged()
+{
+	check(RegionManager);
+	return RegionManager->OnQosRegionIdChanged();
 }
\ No newline at end of file
diff --git a/Engine/Plugins/Online/OnlineFramework/Source/Qos/Private/QosRegionManager.cpp b/Engine/Plugins/Online/OnlineFramework/Source/Qos/Private/QosRegionManager.cpp
--- a/Engine/Plugins/Online/OnlineFramework/Source/Qos/Private/QosRegionManager.cpp
+++ b/Engine/Plugins/Online/OnlineFramework/Source/Qos/Private/QosRegionManager.cpp
@@ -712,9 +712,13 @@ bool UQosRegionManager::SetSelectedRegion(const FString& InRegionId, bool bForce
 			{
 				if (RegionInfo.IsUsable())
 				{
-					UE_LOG(LogQos, Verbose, TEXT("[UQosRegionManager::SetSelectedRegion] Old: \"%s\"  New: \"%s\"  (force? %s)"),
-						*SelectedRegionId, *RegionId, *LexToString(bForce));
+					FString OldRegionId(SelectedRegionId);
 					SelectedRegionId = MoveTemp(RegionId);
+
+					UE_LOG(LogQos, Verbose, TEXT("[UQosRegionManager::SetSelectedRegion] Old: \"%s\"  New: \"%s\"  (force? %s)"),
+						*OldRegionId, *SelectedRegionId, *LexToString(bForce));
+
+					OnQosRegionIdChanged().Broadcast(OldRegionId, SelectedRegionId);
 					return true;
 				}
 				else
diff --git a/Engine/Plugins/Online/OnlineFramework/Source/Qos/Public/QosInterface.h b/Engine/Plugins/Online/OnlineFramework/Source/Qos/Public/QosInterface.h
--- a/Engine/Plugins/Online/OnlineFramework/Source/Qos/Public/QosInterface.h
+++ b/Engine/Plugins/Online/OnlineFramework/Source/Qos/Public/QosInterface.h
@@ -124,6 +124,12 @@ public:
 	*/
 	void RegisterQoSSettingsChangedDelegate(const FSimpleDelegate& OnQoSSettingsChanged);
 
+	/**
+	 * Delegate that fires whenever the current QoS region ID changes.
+	 */
+	 DECLARE_MULTICAST_DELEGATE_TwoParams(FOnQosRegionIdChanged, const FString& /* OldRegionId */, const FString& /* NewRegionId */);
+	 FOnQosRegionIdChanged& OnQosRegionIdChanged();
+
 protected:
 
 	friend class FQosModule;
diff --git a/Engine/Plugins/Online/OnlineFramework/Source/Qos/Public/QosRegionManager.h b/Engine/Plugins/Online/OnlineFramework/Source/Qos/Public/QosRegionManager.h
--- a/Engine/Plugins/Online/OnlineFramework/Source/Qos/Public/QosRegionManager.h
+++ b/Engine/Plugins/Online/OnlineFramework/Source/Qos/Public/QosRegionManager.h
@@ -522,6 +522,12 @@ public:
 	 */
 	void ForceSelectRegion(const FString& RegionId);
 
+	/**
+	 * Delegate that fires whenever the current QoS region ID changes.
+	 */
+	 DECLARE_MULTICAST_DELEGATE_TwoParams(FOnQosRegionIdChanged, const FString& /* OldRegionId */, const FString& /* NewRegionId */);
+	 FOnQosRegionIdChanged& OnQosRegionIdChanged() { return OnQosRegionIdChangedDelegate; }
+
 	/**
 	 * Get the datacenter id for this instance, checking ini and commandline overrides
 	 * This is only relevant for dedicated servers (so they can advertise). 
@@ -689,6 +695,8 @@ private:
 
 	FSimpleDelegate OnQoSSettingsChangedDelegate;
 
+	FOnQosRegionIdChanged OnQosRegionIdChangedDelegate;
+
 	static const TCHAR* SubspaceDelimiterDefault;
 };
 
