Support for new Request-To-Join privacy options.
#OKFORGITHUB PUBLIC

diff --git a/Engine/Plugins/Online/OnlineFramework/Source/Party/Private/SocialToolkit.cpp b/Engine/Plugins/Online/OnlineFramework/Source/Party/Private/SocialToolkit.cpp
--- a/Engine/Plugins/Online/OnlineFramework/Source/Party/Private/SocialToolkit.cpp
+++ b/Engine/Plugins/Online/OnlineFramework/Source/Party/Private/SocialToolkit.cpp
@@ -1028,6 +1028,11 @@ void USocialToolkit::HandleAcceptFriendInviteComplete(int32 LocalUserNum, bool b
 	OnAcceptFriendInviteComplete(InviterUserId, bWasSuccessful, ErrorStr);
 }
 
+const bool USocialToolkit::IsInviteAllowedFromUser(const USocialUser& User, const TSharedRef<const IOnlinePartyJoinInfo>& InvitePtr) const
+{
+  return User.IsFriend(ESocialSubsystem::Primary);
+}
+
 void USocialToolkit::HandlePartyInviteReceived(const FUniqueNetId& LocalUserId, const IOnlinePartyJoinInfo& Invite)
 {
 	if (LocalUserId == GetLocalUserNetId(ESocialSubsystem::Primary))
@@ -1038,7 +1043,7 @@ void USocialToolkit::HandlePartyInviteReceived(const FUniqueNetId& LocalUserId,
 		QueueUserDependentActionInternal(Invite.GetSourceUserId(), ESocialSubsystem::Primary,
 			[this, Invite = Invite.AsShared()] (USocialUser& User)
 			{
-				if (User.IsFriend(ESocialSubsystem::Primary))
+				if (IsInviteAllowedFromUser(User, Invite))
 				{
 #if PARTY_PLATFORM_INVITE_PERMISSIONS
 					CanReceiveInviteFrom(User, Invite, [this, Invite, UserId = User.GetUserId(ESocialSubsystem::Primary)](const bool bResult)
diff --git a/Engine/Plugins/Online/OnlineFramework/Source/Party/Private/User/SocialUser.cpp b/Engine/Plugins/Online/OnlineFramework/Source/Party/Private/User/SocialUser.cpp
--- a/Engine/Plugins/Online/OnlineFramework/Source/Party/Private/User/SocialUser.cpp
+++ b/Engine/Plugins/Online/OnlineFramework/Source/Party/Private/User/SocialUser.cpp
@@ -1083,6 +1083,11 @@ void USocialUser::HandleRequestToJoinRemoved(const IOnlinePartyRequestToJoinInfo
 	NotifyRequestToJoinRemoved(Request, Reason);
 }
 
+void USocialUser::RequestToJoinParty()
+{
+	return GetOwningToolkit().RequestToJoinParty(*this);
+}
+
 void USocialUser::AcceptRequestToJoinParty() const
 {
 	if (USocialParty* Party = GetOwningToolkit().GetSocialManager().GetParty(IOnlinePartySystem::GetPrimaryPartyTypeId()))
diff --git a/Engine/Plugins/Online/OnlineFramework/Source/Party/Public/SocialToolkit.h b/Engine/Plugins/Online/OnlineFramework/Source/Party/Public/SocialToolkit.h
--- a/Engine/Plugins/Online/OnlineFramework/Source/Party/Public/SocialToolkit.h
+++ b/Engine/Plugins/Online/OnlineFramework/Source/Party/Public/SocialToolkit.h
@@ -199,6 +199,9 @@ protected:
 	bool bQueryBlockedPlayersOnStartup = true;
 	bool bQueryRecentPlayersOnStartup = true;
 
+	/** Whether we allow other users to send us invites */
+	virtual const bool IsInviteAllowedFromUser(const USocialUser& User, const TSharedRef<const IOnlinePartyJoinInfo>& InviteRef) const;
+
 private:
 	void QueueUserDependentActionInternal(const FUniqueNetIdRepl& SubsystemId, ESocialSubsystem SubsystemType, TFunction<void(USocialUser&)>&& UserActionFunc, bool bExecutePostInit = true);
 
