Add completion delegates to send and accept friend request

#tests added the following code to try send invite and accept invite in SocialUser.cpp, sent and accepted friend invites and observed logs:
```
bool USocialUser::SendFriendInvite(ESocialSubsystem SubsystemType)
{
	return GetOwningToolkit().TrySendFriendInvite(*this, SubsystemType, FOnTrySendFriendInviteComplete::CreateLambda([](bool bWasSuccessful, FString Error)
		{
			UE_LOG(LogTemp, Warning, TEXT("Send Invite Completion Works!!! bWasSuccessful: %s, Error: [%s]"), *LexToString(bWasSuccessful), *Error)
		})
	);
}

bool USocialUser::AcceptFriendInvite(ESocialSubsystem SocialSubsystem) const
{
	return GetOwningToolkit().AcceptFriendInvite(*this, SocialSubsystem, FOnAcceptFriendInviteComplete::CreateLambda([](bool bWasSuccessful, FString Error)
		{
			UE_LOG(LogTemp, Warning, TEXT("Accept Invite Completion Works!!! bWasSuccessful: %s, Error: [%s]"), *LexToString(bWasSuccessful), *Error)
		})
	);
}

```

diff --git a/Engine/Plugins/Online/OnlineFramework/Source/Party/Private/SocialToolkit.cpp b/Engine/Plugins/Online/OnlineFramework/Source/Party/Private/SocialToolkit.cpp
--- a/Engine/Plugins/Online/OnlineFramework/Source/Party/Private/SocialToolkit.cpp
+++ b/Engine/Plugins/Online/OnlineFramework/Source/Party/Private/SocialToolkit.cpp
@@ -334,15 +334,19 @@ USocialUser* USocialToolkit::FindUser(const FUniqueNetIdRepl& UserId) const
 	return FoundUser ? FoundUser->Get() : nullptr;
 }
 
-void USocialToolkit::TrySendFriendInvite(const FString& DisplayNameOrEmail) const
+void USocialToolkit::TrySendFriendInvite(const FString& DisplayNameOrEmail, FOnTrySendFriendInviteComplete OnCompleteDelegate) const
 {
 	IOnlineSubsystem* PrimaryOSS = GetSocialOss(ESocialSubsystem::Primary);
 	IOnlineUserPtr UserInterface = PrimaryOSS ? PrimaryOSS->GetUserInterface() : nullptr;
 	if (UserInterface.IsValid())
 	{
-		IOnlineUser::FOnQueryUserMappingComplete QueryCompleteDelegate = IOnlineUser::FOnQueryUserMappingComplete::CreateUObject(const_cast<USocialToolkit*>(this), &USocialToolkit::HandleQueryPrimaryUserIdMappingComplete);
+		IOnlineUser::FOnQueryUserMappingComplete QueryCompleteDelegate = IOnlineUser::FOnQueryUserMappingComplete::CreateUObject(const_cast<USocialToolkit*>(this), &USocialToolkit::HandleQueryPrimaryUserIdMappingComplete, OnCompleteDelegate);
 		UserInterface->QueryUserIdMapping(*GetLocalUserNetId(ESocialSubsystem::Primary), DisplayNameOrEmail, QueryCompleteDelegate);
 	}
+	else
+	{
+		OnCompleteDelegate.ExecuteIfBound(false, FriendInviteFailureReason::InviteFailReason_Uninitialized);
+	}
 }
 
 bool USocialToolkit::GetAuthAttribute(ESocialSubsystem SubsystemType, const FString& AttributeKey, FString& OutValue) const
@@ -509,38 +513,57 @@ void USocialToolkit::NotifySubsystemIdEstablished(USocialUser& SocialUser, ESoci
 	}
 }
 
-bool USocialToolkit::TrySendFriendInvite(USocialUser& SocialUser, ESocialSubsystem SubsystemType) const
+bool USocialToolkit::TrySendFriendInvite(USocialUser& SocialUser, ESocialSubsystem SubsystemType, FOnTrySendFriendInviteComplete OnCompleteDelegate) const
 {
 	if (SocialUser.GetFriendInviteStatus(SubsystemType) == EInviteStatus::PendingOutbound)
 	{
 		OnFriendInviteSent().Broadcast(SocialUser, SubsystemType);
+		OnCompleteDelegate.ExecuteIfBound(true, FString());
 		return true;
 	}
 	else if (!SocialUser.IsFriend(SubsystemType))
 	{
 		IOnlineFriendsPtr FriendsInterface = Online::GetFriendsInterface(GetWorld(), USocialManager::GetSocialOssName(SubsystemType));
 		const FUniqueNetIdRepl SubsystemId = SocialUser.GetUserId(SubsystemType);
-
-		const bool bIsFriendshipRestricted = IsFriendshipRestricted(SocialUser, SubsystemType);
-
-		if (FriendsInterface && SubsystemId.IsValid() && !bIsFriendshipRestricted)
+		if (FriendsInterface && SubsystemId.IsValid())
+		{
+			const bool bIsFriendshipRestricted = IsFriendshipRestricted(SocialUser, SubsystemType);
+			if (!bIsFriendshipRestricted)
+			{
+				return FriendsInterface->SendInvite(GetLocalUserNum(), *SubsystemId, FriendListToQuery, FOnSendInviteComplete::CreateUObject(const_cast<USocialToolkit*>(this), &USocialToolkit::HandleFriendInviteSent, SubsystemType, SocialUser.GetDisplayName(), OnCompleteDelegate));
+			}
+			else
+			{
+				OnCompleteDelegate.ExecuteIfBound(false, FriendInviteFailureReason::InviteFailReason_FriendshipRestricted);
+				return false;
+			}
+		}
+		else
 		{
-			return FriendsInterface->SendInvite(GetLocalUserNum(), *SubsystemId, FriendListToQuery, FOnSendInviteComplete::CreateUObject(const_cast<USocialToolkit*>(this), &USocialToolkit::HandleFriendInviteSent, SubsystemType, SocialUser.GetDisplayName()));
+			OnCompleteDelegate.ExecuteIfBound(false, FriendInviteFailureReason::InviteFailReason_Uninitialized);
+			return false;
 		}
 	}
-	return false;
+	else
+	{
+		OnCompleteDelegate.ExecuteIfBound(false, FriendInviteFailureReason::InviteFailReason_AlreadyFriends);
+		return false;
+	}
 }
 
-bool USocialToolkit::AcceptFriendInvite(const USocialUser& SocialUser, ESocialSubsystem SubsystemType) const
+bool USocialToolkit::AcceptFriendInvite(const USocialUser& SocialUser, ESocialSubsystem SubsystemType, FOnAcceptFriendInviteComplete OnCompleteDelegate) const
 {
 	if (SocialUser.GetFriendInviteStatus(SubsystemType) == EInviteStatus::PendingInbound)
 	{
 		IOnlineFriendsPtr FriendsInterface = GetSocialOss(SubsystemType)->GetFriendsInterface();
 		check(FriendsInterface.IsValid());
-
-		return FriendsInterface->AcceptInvite(GetLocalUserNum(), *SocialUser.GetUserId(SubsystemType), EFriendsLists::ToString(EFriendsLists::Default), FOnAcceptInviteComplete::CreateUObject(const_cast<USocialToolkit*>(this), &USocialToolkit::HandleAcceptFriendInviteComplete));
+		return FriendsInterface->AcceptInvite(GetLocalUserNum(), *SocialUser.GetUserId(SubsystemType), EFriendsLists::ToString(EFriendsLists::Default), FOnAcceptInviteComplete::CreateUObject(const_cast<USocialToolkit*>(this), &USocialToolkit::HandleAcceptFriendInviteComplete, OnCompleteDelegate));
+	}
+	else
+	{
+		OnCompleteDelegate.ExecuteIfBound(false, FriendAcceptFailureReason::AcceptFailReason_NotPendingInbound);
+		return false;
 	}
-	return false;
 }
 
 bool USocialToolkit::IsFriendshipRestricted(const USocialUser& SocialUser, ESocialSubsystem SubsystemType) const
@@ -892,32 +915,36 @@ void USocialToolkit::HandlePresenceReceived(const FUniqueNetId& UserId, const TS
 	}
 }
 
-void USocialToolkit::HandleQueryPrimaryUserIdMappingComplete(bool bWasSuccessful, const FUniqueNetId& RequestingUserId, const FString& DisplayName, const FUniqueNetId& IdentifiedUserId, const FString& Error)
+void USocialToolkit::HandleQueryPrimaryUserIdMappingComplete(bool bWasSuccessful, const FUniqueNetId& RequestingUserId, const FString& DisplayName, const FUniqueNetId& IdentifiedUserId, const FString& Error, FOnTrySendFriendInviteComplete OnCompleteDelegate)
 {
 	if (!IdentifiedUserId.IsValid())
 	{
 		OnSendFriendInviteComplete(IdentifiedUserId, DisplayName, false, FriendInviteFailureReason::InviteFailReason_NotFound);
+		OnCompleteDelegate.ExecuteIfBound(false, FriendInviteFailureReason::InviteFailReason_NotFound);
 	}
 	else if (RequestingUserId == IdentifiedUserId)
 	{
 		OnSendFriendInviteComplete(IdentifiedUserId, DisplayName, false, FriendInviteFailureReason::InviteFailReason_AddingSelfFail);
+		OnCompleteDelegate.ExecuteIfBound(false, FriendInviteFailureReason::InviteFailReason_AddingSelfFail);
 	}
 	else
 	{
 		QueueUserDependentActionInternal(IdentifiedUserId.AsShared(), ESocialSubsystem::Primary,
-			[this, DisplayName] (USocialUser& SocialUser)
+			[this, DisplayName, OnCompleteDelegate] (USocialUser& SocialUser)
 			{
 				if (SocialUser.IsBlocked())
 				{
 					OnSendFriendInviteComplete(*SocialUser.GetUserId(ESocialSubsystem::Primary), DisplayName, false, FriendInviteFailureReason::InviteFailReason_AddingBlockedFail);
+					OnCompleteDelegate.ExecuteIfBound(false, FriendInviteFailureReason::InviteFailReason_AddingBlockedFail);
 				}
 				else if (SocialUser.IsFriend(ESocialSubsystem::Primary))
 				{
 					OnSendFriendInviteComplete(*SocialUser.GetUserId(ESocialSubsystem::Primary), DisplayName, false, FriendInviteFailureReason::InviteFailReason_AlreadyFriends);
+					OnCompleteDelegate.ExecuteIfBound(false, FriendInviteFailureReason::InviteFailReason_AlreadyFriends);
 				}
 				else
 				{
-					TrySendFriendInvite(SocialUser, ESocialSubsystem::Primary);
+					TrySendFriendInvite(SocialUser, ESocialSubsystem::Primary, OnCompleteDelegate);
 				}
 			});
 	}
@@ -976,7 +1003,7 @@ void USocialToolkit::HandleFriendInviteRejected(const FUniqueNetId& LocalUserId,
 	}
 }
 
-void USocialToolkit::HandleFriendInviteSent(int32 LocalUserNum, bool bWasSuccessful, const FUniqueNetId& InvitedUserId, const FString& ListName, const FString& ErrorStr, ESocialSubsystem SubsystemType, FString DisplayName)
+void USocialToolkit::HandleFriendInviteSent(int32 LocalUserNum, bool bWasSuccessful, const FUniqueNetId& InvitedUserId, const FString& ListName, const FString& ErrorStr, ESocialSubsystem SubsystemType, FString DisplayName, FOnTrySendFriendInviteComplete OnCompleteDelegate)
 {
 	if (bWasSuccessful)
 	{
@@ -995,6 +1022,7 @@ void USocialToolkit::HandleFriendInviteSent(int32 LocalUserNum, bool bWasSuccess
 			});
 	}
 	OnSendFriendInviteComplete(InvitedUserId, DisplayName, bWasSuccessful, ErrorStr);
+	OnCompleteDelegate.ExecuteIfBound(bWasSuccessful, ErrorStr);
 }
 
 void USocialToolkit::HandleFriendRemoved(const FUniqueNetId& LocalUserId, const FUniqueNetId& FriendId, ESocialSubsystem SubsystemType)
@@ -1023,9 +1051,10 @@ void USocialToolkit::HandleDeleteFriendComplete(int32 InLocalUserNum, bool bWasS
 	OnDeleteFriendComplete(InLocalUserNum, bWasSuccessful, DeletedFriendId, ListName, ErrorStr, SubsystemType);
 }
 
-void USocialToolkit::HandleAcceptFriendInviteComplete(int32 LocalUserNum, bool bWasSuccessful, const FUniqueNetId& InviterUserId, const FString& ListName, const FString& ErrorStr)
+void USocialToolkit::HandleAcceptFriendInviteComplete(int32 LocalUserNum, bool bWasSuccessful, const FUniqueNetId& InviterUserId, const FString& ListName, const FString& ErrorStr, FOnAcceptFriendInviteComplete OnCompleteDelegate)
 {
 	OnAcceptFriendInviteComplete(InviterUserId, bWasSuccessful, ErrorStr);
+	OnCompleteDelegate.ExecuteIfBound(bWasSuccessful, ErrorStr);
 }
 
 const bool USocialToolkit::IsInviteAllowedFromUser(const USocialUser& User, const TSharedRef<const IOnlinePartyJoinInfo>& InvitePtr) const
diff --git a/Engine/Plugins/Online/OnlineFramework/Source/Party/Public/SocialToolkit.h b/Engine/Plugins/Online/OnlineFramework/Source/Party/Public/SocialToolkit.h
--- a/Engine/Plugins/Online/OnlineFramework/Source/Party/Public/SocialToolkit.h
+++ b/Engine/Plugins/Online/OnlineFramework/Source/Party/Public/SocialToolkit.h
@@ -30,9 +30,17 @@ namespace FriendInviteFailureReason
 	const FString InviteFailReason_AddingSelfFail = TEXT("AddingSelfFail");
 	const FString InviteFailReason_AddingBlockedFail = TEXT("AddingBlockedFail");
 	const FString InviteFailReason_AlreadyFriends = TEXT("AlreadyFriends");
+	const FString InviteFailReason_Uninitialized = TEXT("Uninitialized");
+	const FString InviteFailReason_FriendshipRestricted = TEXT("FriendshipRestricted");
+}
+namespace FriendAcceptFailureReason
+{
+	const FString AcceptFailReason_NotPendingInbound = TEXT("NotPendingInbound");
 }
 
 DECLARE_DELEGATE_OneParam(FUserDependentAction, USocialUser&);
+DECLARE_DELEGATE_TwoParams(FOnTrySendFriendInviteComplete, bool /* bWasSuccessful */, FString /* ErrorString */);
+DECLARE_DELEGATE_TwoParams(FOnAcceptFriendInviteComplete, bool /* bWasSuccessful */, FString /* ErrorString */);
 
 /** Represents the full suite of social functionality available to a given LocalPlayer */
 UCLASS(Within = SocialManager)
@@ -88,12 +96,12 @@ public:
 	 */
 	void QueueUserDependentAction(const FUniqueNetIdRepl& UserId, TFunction<void(USocialUser&)>&& UserActionFunc, bool bExecutePostInit = true);
 	void QueueUserDependentAction(const FUniqueNetIdRepl& SubsystemId, FUserDependentAction UserActionDelegate);
-
+	
 	/**
 	 * Attempts to send a friend invite to another user based on display name or email.
 	 * Only necessary to use this path when you do not have a known USocialUser for this user already.
 	 */
-	void TrySendFriendInvite(const FString& DisplayNameOrEmail) const;
+	virtual void TrySendFriendInvite(const FString& DisplayNameOrEmail, FOnTrySendFriendInviteComplete OnCompleteDelegate = FOnTrySendFriendInviteComplete()) const;
 
 	virtual bool IsFriendshipRestricted(const USocialUser& SocialUser, ESocialSubsystem SubsystemType) const;
 
@@ -143,9 +151,9 @@ PACKAGE_SCOPE:
 	void NotifySubsystemIdEstablished(USocialUser& SocialUser, ESocialSubsystem SubsystemType, const FUniqueNetIdRepl& SubsystemId);
 	TSubclassOf<USocialChatManager> GetChatManagerClass() { return ChatManagerClass; }
 
-	bool TrySendFriendInvite(USocialUser& SocialUser, ESocialSubsystem SubsystemType) const;
+	bool TrySendFriendInvite(USocialUser& SocialUser, ESocialSubsystem SubsystemType, FOnTrySendFriendInviteComplete OnCompleteDelegate = FOnTrySendFriendInviteComplete()) const;
 
-	bool AcceptFriendInvite(const USocialUser& SocialUser, ESocialSubsystem SubsystemType) const;
+	bool AcceptFriendInvite(const USocialUser& SocialUser, ESocialSubsystem SubsystemType, FOnAcceptFriendInviteComplete OnCompleteDelegate = FOnAcceptFriendInviteComplete()) const;
 
 	void RequestToJoinParty(USocialUser& SocialUser);
 
@@ -240,11 +248,11 @@ private:	// Handlers
 	void HandleFriendInviteReceived(const FUniqueNetId& LocalUserId, const FUniqueNetId& SenderId, ESocialSubsystem SubsystemType);
 	void HandleFriendInviteAccepted(const FUniqueNetId& LocalUserId, const FUniqueNetId& NewFriendId, ESocialSubsystem SubsystemType);
 	void HandleFriendInviteRejected(const FUniqueNetId& LocalUserId, const FUniqueNetId& RejecterId, ESocialSubsystem SubsystemType);
-	void HandleFriendInviteSent(int32 LocalUserNum, bool bWasSuccessful, const FUniqueNetId& InvitedUserId, const FString& ListName, const FString& ErrorStr, ESocialSubsystem SubsystemType, FString DisplayName);
+	void HandleFriendInviteSent(int32 LocalUserNum, bool bWasSuccessful, const FUniqueNetId& InvitedUserId, const FString& ListName, const FString& ErrorStr, ESocialSubsystem SubsystemType, FString DisplayName, FOnTrySendFriendInviteComplete OnCompleteDelegate);
 	void HandleFriendRemoved(const FUniqueNetId& LocalUserId, const FUniqueNetId& FormerFriendId, ESocialSubsystem SubsystemType);
 
 	void HandleDeleteFriendComplete(int32 LocalPlayer, bool bWasSuccessful, const FUniqueNetId& FormerFriendId, const FString& ListName, const FString& ErrorStr, ESocialSubsystem SubsystemType);
-	void HandleAcceptFriendInviteComplete(int32 LocalUserNum, bool bWasSuccessful, const FUniqueNetId& InviterUserId, const FString& ListName, const FString& ErrorStr);
+	void HandleAcceptFriendInviteComplete(int32 LocalUserNum, bool bWasSuccessful, const FUniqueNetId& InviterUserId, const FString& ListName, const FString& ErrorStr, FOnAcceptFriendInviteComplete OnCompleteDelegate);
 
 	void HandlePartyInviteReceived(const FUniqueNetId& LocalUserId, const IOnlinePartyJoinInfo& Invite);
 	void HandlePartyInviteRemoved(const FUniqueNetId& LocalUserId, const IOnlinePartyJoinInfo& Invite, EPartyInvitationRemovedReason Reason);
@@ -255,7 +263,7 @@ private:	// Handlers
 	void HandleBlockPlayerComplete(int32 LocalUserNum, bool bWasSuccessful, const FUniqueNetId& BlockedPlayerID, const FString& ListName, const FString& ErrorStr, ESocialSubsystem SubsystemType);
 	void HandleUnblockPlayerComplete(int32 LocalUserNum, bool bWasSuccessful, const FUniqueNetId& UnblockedPlayerID, const FString& ListName, const FString& ErrorStr, ESocialSubsystem SubsystemType);
 	
-	void HandleQueryPrimaryUserIdMappingComplete(bool bWasSuccessful, const FUniqueNetId& RequestingUserId, const FString& DisplayName, const FUniqueNetId& IdentifiedUserId, const FString& Error);
+	void HandleQueryPrimaryUserIdMappingComplete(bool bWasSuccessful, const FUniqueNetId& RequestingUserId, const FString& DisplayName, const FUniqueNetId& IdentifiedUserId, const FString& Error, FOnTrySendFriendInviteComplete OnCompleteDelegate);
 
 	void HandlePartyMemberExited(const FUniqueNetId& LocalUserId, const FOnlinePartyId& PartyId, const FUniqueNetId& MemberId, const EMemberExitedReason Reason);
 	void HandleGameDestroyed(const FName SessionName, bool bWasSuccessful);
