Move Add and Accept friend invite methods to public and virtual

Allow for custom game level functionality

#tests full compile

diff --git a/Engine/Plugins/Online/OnlineFramework/Source/Party/Public/SocialToolkit.h b/Engine/Plugins/Online/OnlineFramework/Source/Party/Public/SocialToolkit.h
--- a/Engine/Plugins/Online/OnlineFramework/Source/Party/Public/SocialToolkit.h
+++ b/Engine/Plugins/Online/OnlineFramework/Source/Party/Public/SocialToolkit.h
@@ -103,6 +103,11 @@ public:
 	 */
 	virtual void TrySendFriendInvite(const FString& DisplayNameOrEmail, FOnTrySendFriendInviteComplete OnCompleteDelegate = FOnTrySendFriendInviteComplete()) const;
 
+	/** Send friend invite to a SocialUser */
+	virtual bool TrySendFriendInvite(USocialUser& SocialUser, ESocialSubsystem SubsystemType, FOnTrySendFriendInviteComplete OnCompleteDelegate = FOnTrySendFriendInviteComplete()) const;
+
+	virtual bool AcceptFriendInvite(const USocialUser& SocialUser, ESocialSubsystem SubsystemType, FOnAcceptFriendInviteComplete OnCompleteDelegate = FOnAcceptFriendInviteComplete()) const;
+
 	virtual bool IsFriendshipRestricted(const USocialUser& SocialUser, ESocialSubsystem SubsystemType) const;
 
 	bool GetAuthAttribute(ESocialSubsystem SubsystemType, const FString& AttributeKey, FString& OutValue) const;
@@ -151,10 +156,6 @@ PACKAGE_SCOPE:
 	void NotifySubsystemIdEstablished(USocialUser& SocialUser, ESocialSubsystem SubsystemType, const FUniqueNetIdRepl& SubsystemId);
 	TSubclassOf<USocialChatManager> GetChatManagerClass() { return ChatManagerClass; }
 
-	bool TrySendFriendInvite(USocialUser& SocialUser, ESocialSubsystem SubsystemType, FOnTrySendFriendInviteComplete OnCompleteDelegate = FOnTrySendFriendInviteComplete()) const;
-
-	bool AcceptFriendInvite(const USocialUser& SocialUser, ESocialSubsystem SubsystemType, FOnAcceptFriendInviteComplete OnCompleteDelegate = FOnAcceptFriendInviteComplete()) const;
-
 	void RequestToJoinParty(USocialUser& SocialUser);
 
 	void HandleUserInvalidated(USocialUser& InvalidUser);
