Refactor friend invite call location from business logic of TrySendFriendInvite

This will allow for better control over friend invites

#tests full compile, sent and accepted friend invites

diff --git a/Engine/Plugins/Online/OnlineFramework/Source/Party/Private/SocialToolkit.cpp b/Engine/Plugins/Online/OnlineFramework/Source/Party/Private/SocialToolkit.cpp
--- a/Engine/Plugins/Online/OnlineFramework/Source/Party/Private/SocialToolkit.cpp
+++ b/Engine/Plugins/Online/OnlineFramework/Source/Party/Private/SocialToolkit.cpp
@@ -523,24 +523,14 @@ bool USocialToolkit::TrySendFriendInvite(USocialUser& SocialUser, ESocialSubsyst
 	}
 	else if (!SocialUser.IsFriend(SubsystemType))
 	{
-		IOnlineFriendsPtr FriendsInterface = Online::GetFriendsInterface(GetWorld(), USocialManager::GetSocialOssName(SubsystemType));
-		const FUniqueNetIdRepl SubsystemId = SocialUser.GetUserId(SubsystemType);
-		if (FriendsInterface && SubsystemId.IsValid())
+		const bool bIsFriendshipRestricted = IsFriendshipRestricted(SocialUser, SubsystemType);
+		if (!bIsFriendshipRestricted)
 		{
-			const bool bIsFriendshipRestricted = IsFriendshipRestricted(SocialUser, SubsystemType);
-			if (!bIsFriendshipRestricted)
-			{
-				return FriendsInterface->SendInvite(GetLocalUserNum(), *SubsystemId, FriendListToQuery, FOnSendInviteComplete::CreateUObject(const_cast<USocialToolkit*>(this), &USocialToolkit::HandleFriendInviteSent, SubsystemType, SocialUser.GetDisplayName(), OnCompleteDelegate));
-			}
-			else
-			{
-				OnCompleteDelegate.ExecuteIfBound(false, FriendInviteFailureReason::InviteFailReason_FriendshipRestricted);
-				return false;
-			}
+			return InternalSendFriendInvite(SocialUser, SubsystemType, OnCompleteDelegate);
 		}
 		else
 		{
-			OnCompleteDelegate.ExecuteIfBound(false, FriendInviteFailureReason::InviteFailReason_Uninitialized);
+			OnCompleteDelegate.ExecuteIfBound(false, FriendInviteFailureReason::InviteFailReason_FriendshipRestricted);
 			return false;
 		}
 	}
@@ -551,6 +541,23 @@ bool USocialToolkit::TrySendFriendInvite(USocialUser& SocialUser, ESocialSubsyst
 	}
 }
 
+
+bool USocialToolkit::InternalSendFriendInvite(USocialUser& SocialUser, ESocialSubsystem SubsystemType, FOnTrySendFriendInviteComplete OnCompleteDelegate) const
+{
+	IOnlineFriendsPtr FriendsInterface = Online::GetFriendsInterface(GetWorld(), USocialManager::GetSocialOssName(SubsystemType));
+	const FUniqueNetIdRepl SocialUserSubsystemId = SocialUser.GetUserId(SubsystemType);
+	if (FriendsInterface && SocialUserSubsystemId.IsValid())
+	{
+		return FriendsInterface->SendInvite(GetLocalUserNum(), *SocialUserSubsystemId, FriendListToQuery, FOnSendInviteComplete::CreateUObject(const_cast<USocialToolkit*>(this), &USocialToolkit::HandleFriendInviteSent, SubsystemType, SocialUser.GetDisplayName(), OnCompleteDelegate));
+	}
+	else
+	{
+		OnCompleteDelegate.ExecuteIfBound(false, FriendInviteFailureReason::InviteFailReason_Uninitialized);
+		return false;
+	}
+}
+
+
 bool USocialToolkit::AcceptFriendInvite(const USocialUser& SocialUser, ESocialSubsystem SubsystemType, FOnAcceptFriendInviteComplete OnCompleteDelegate) const
 {
 	if (SocialUser.GetFriendInviteStatus(SubsystemType) == EInviteStatus::PendingInbound)
diff --git a/Engine/Plugins/Online/OnlineFramework/Source/Party/Public/SocialToolkit.h b/Engine/Plugins/Online/OnlineFramework/Source/Party/Public/SocialToolkit.h
--- a/Engine/Plugins/Online/OnlineFramework/Source/Party/Public/SocialToolkit.h
+++ b/Engine/Plugins/Online/OnlineFramework/Source/Party/Public/SocialToolkit.h
@@ -165,6 +165,8 @@ PACKAGE_SCOPE:
 #endif
 
 protected:
+	/** protected function that only handles sending a friend request via the friends interface, assumes all checks are handled previously **/
+	virtual bool InternalSendFriendInvite(USocialUser& SocialUser, ESocialSubsystem SubsystemType, FOnTrySendFriendInviteComplete OnCompleteDelegate = FOnTrySendFriendInviteComplete()) const;
 
 	virtual void OnOwnerLoggedIn();
 	virtual void OnOwnerLoggedOut();
