(Social/Party Signal)  Automatically dismiss Requests To Join from blocked users.

diff --git a/Engine/Plugins/Online/OnlineFramework/Source/Party/Private/SocialToolkit.cpp b/Engine/Plugins/Online/OnlineFramework/Source/Party/Private/SocialToolkit.cpp
--- a/Engine/Plugins/Online/OnlineFramework/Source/Party/Private/SocialToolkit.cpp
+++ b/Engine/Plugins/Online/OnlineFramework/Source/Party/Private/SocialToolkit.cpp
@@ -1237,8 +1237,17 @@ void USocialToolkit::HandlePartyRequestToJoinReceived(const FUniqueNetId& LocalU
 		QueueUserDependentActionInternal(RequesterId.AsShared(), ESocialSubsystem::Primary,
 			[this, RequestRef = Request.AsShared()](USocialUser& User)
 			{
-				User.HandleRequestToJoinReceived(*RequestRef);
-				OnPartyRequestToJoinReceived().Broadcast(User, RequestRef);
+				// The requesting user won't know they're blocked, so we can't prevent their request from being sent.
+				// Instead, ignore the request at the receiving end.
+				if (!User.IsBlocked())
+				{
+					User.HandleRequestToJoinReceived(*RequestRef);
+					OnPartyRequestToJoinReceived().Broadcast(User, RequestRef);
+				}
+				else
+				{
+					UE_LOG(LogParty, VeryVerbose, TEXT("%s - Join request from blocked user [%s] ignored"), ANSI_TO_TCHAR(__FUNCTION__), *User.GetUserId(ESocialSubsystem::Primary).ToDebugString());
+				}
 			});
 	}
 }
