Add optional MetaData parameter to the party invitation sending flow.

#tests Compiled DevelopmentEditor Win64 / Test Client / Single file on cpps / Ran a PIE instance and verified that metadata made its way through the invitation process and to the recipient.

diff --git a/Engine/Plugins/Online/OnlineFramework/Source/Party/Private/Party/SocialParty.cpp b/Engine/Plugins/Online/OnlineFramework/Source/Party/Private/Party/SocialParty.cpp
--- a/Engine/Plugins/Online/OnlineFramework/Source/Party/Private/Party/SocialParty.cpp
+++ b/Engine/Plugins/Online/OnlineFramework/Source/Party/Private/Party/SocialParty.cpp
@@ -318,7 +318,7 @@ bool USocialParty::IsInviteRateLimited(const USocialUser& User, ESocialSubsystem
 	return false;
 }
 
-bool USocialParty::TryInviteUser(const USocialUser& UserToInvite, const ESocialPartyInviteMethod InviteMethod)
+bool USocialParty::TryInviteUser(const USocialUser& UserToInvite, const ESocialPartyInviteMethod InviteMethod, const FString& MetaData)
 {
 	bool bSentInvite = false;
 	ESocialPartyInviteFailureReason CanInviteResult = CanInviteUserInternal(UserToInvite);
@@ -369,8 +369,9 @@ bool USocialParty::TryInviteUser(const USocialUser& UserToInvite, const ESocialP
 		if ((!bSentInvite || bMustSendPrimaryInvite) && UserPrimaryId.IsValid() && !IsInviteRateLimited(UserToInvite, ESocialSubsystem::Primary))
 		{
 			// Primary subsystem invites can be sent directly to the user via the party interface
+			const FPartyInvitationRecipient Recipient(*UserPrimaryId, MetaData);
 			const IOnlinePartyPtr PartyInterface = Online::GetPartyInterfaceChecked(GetWorld());
-			const bool bSentPrimaryInvite = PartyInterface->SendInvitation(*OwningLocalUserId, GetPartyId(), *UserPrimaryId);
+			const bool bSentPrimaryInvite = PartyInterface->SendInvitation(*OwningLocalUserId, GetPartyId(), Recipient);
 			ESocialPartyInviteFailureReason FailureReason = bSentPrimaryInvite ? ESocialPartyInviteFailureReason::Success : ESocialPartyInviteFailureReason::PartyInviteFailed;
 			OnInviteSentInternal(ESocialSubsystem::Primary, UserToInvite, bSentPrimaryInvite, FailureReason, InviteMethod);
 			bSentInvite |= bSentPrimaryInvite;
diff --git a/Engine/Plugins/Online/OnlineFramework/Source/Party/Private/User/SocialUser.cpp b/Engine/Plugins/Online/OnlineFramework/Source/Party/Private/User/SocialUser.cpp
--- a/Engine/Plugins/Online/OnlineFramework/Source/Party/Private/User/SocialUser.cpp
+++ b/Engine/Plugins/Online/OnlineFramework/Source/Party/Private/User/SocialUser.cpp
@@ -930,11 +930,11 @@ bool USocialUser::CanInviteToParty(const FOnlinePartyTypeId& PartyTypeId) const
 	return false;
 }
 
-bool USocialUser::InviteToParty(const FOnlinePartyTypeId& PartyTypeId, const ESocialPartyInviteMethod InviteMethod) const
+bool USocialUser::InviteToParty(const FOnlinePartyTypeId& PartyTypeId, const ESocialPartyInviteMethod InviteMethod, const FString& MetaData) const
 {
 	if (USocialParty* Party = GetOwningToolkit().GetSocialManager().GetParty(PartyTypeId))
 	{
-		return Party->TryInviteUser(*this, InviteMethod);
+		return Party->TryInviteUser(*this, InviteMethod, MetaData);
 	}
 	return false;
 }
@@ -1092,7 +1092,7 @@ void USocialUser::AcceptRequestToJoinParty() const
 {
 	if (USocialParty* Party = GetOwningToolkit().GetSocialManager().GetParty(IOnlinePartySystem::GetPrimaryPartyTypeId()))
 	{
-		Party->TryInviteUser(*this);
+		Party->TryInviteUser(*this, ESocialPartyInviteMethod::Other, TEXT("RequestToJoin"));
 		IOnlinePartyPtr PartyInterface = Online::GetPartyInterfaceChecked(GetWorld());
 		PartyInterface->ClearRequestToJoinParty(*GetOwningToolkit().GetLocalUserNetId(ESocialSubsystem::Primary), Party->GetPartyId(), *GetUserId(ESocialSubsystem::Primary), EPartyRequestToJoinRemovedReason::Accepted);
 	}
diff --git a/Engine/Plugins/Online/OnlineFramework/Source/Party/Public/Party/SocialParty.h b/Engine/Plugins/Online/OnlineFramework/Source/Party/Public/Party/SocialParty.h
--- a/Engine/Plugins/Online/OnlineFramework/Source/Party/Public/Party/SocialParty.h
+++ b/Engine/Plugins/Online/OnlineFramework/Source/Party/Public/Party/SocialParty.h
@@ -229,7 +229,7 @@ PACKAGE_SCOPE:
 	bool CanPromoteMember(const UPartyMember& PartyMember) const;
 	bool CanKickMember(const UPartyMember& PartyMember) const;
 	
-	bool TryInviteUser(const USocialUser& UserToInvite, const ESocialPartyInviteMethod InviteMethod = ESocialPartyInviteMethod::Other);
+	bool TryInviteUser(const USocialUser& UserToInvite, const ESocialPartyInviteMethod InviteMethod = ESocialPartyInviteMethod::Other, const FString& MetaData = FString());
 	bool TryPromoteMember(const UPartyMember& PartyMember);
 	bool TryKickMember(const UPartyMember& PartyMember);
 	//--------------------------
diff --git a/Engine/Plugins/Online/OnlineFramework/Source/Party/Public/User/SocialUser.h b/Engine/Plugins/Online/OnlineFramework/Source/Party/Public/User/SocialUser.h
--- a/Engine/Plugins/Online/OnlineFramework/Source/Party/Public/User/SocialUser.h
+++ b/Engine/Plugins/Online/OnlineFramework/Source/Party/Public/User/SocialUser.h
@@ -124,7 +124,7 @@ public:
 
 	bool HasBeenInvitedToParty(const FOnlinePartyTypeId& PartyTypeId) const;
 	bool CanInviteToParty(const FOnlinePartyTypeId& PartyTypeId) const;
-	bool InviteToParty(const FOnlinePartyTypeId& PartyTypeId, const ESocialPartyInviteMethod InviteMethod = ESocialPartyInviteMethod::Other) const;
+	bool InviteToParty(const FOnlinePartyTypeId& PartyTypeId, const ESocialPartyInviteMethod InviteMethod = ESocialPartyInviteMethod::Other, const FString& MetaData = FString()) const;
 
 	virtual bool BlockUser(ESocialSubsystem Subsystem) const;
 	virtual bool UnblockUser(ESocialSubsystem Subsystem) const;
diff --git a/Engine/Plugins/Online/OnlineSubsystem/Source/Public/Interfaces/OnlinePartyInterface.h b/Engine/Plugins/Online/OnlineSubsystem/Source/Public/Interfaces/OnlinePartyInterface.h
--- a/Engine/Plugins/Online/OnlineSubsystem/Source/Public/Interfaces/OnlinePartyInterface.h
+++ b/Engine/Plugins/Online/OnlineSubsystem/Source/Public/Interfaces/OnlinePartyInterface.h
@@ -729,6 +729,12 @@ struct FPartyInvitationRecipient
 		: Id(InId.AsShared())
 	{}
 
+	/** Constructor */
+	FPartyInvitationRecipient(const FUniqueNetId& InId, const FString& InMetaData)
+		: Id(InId.AsShared())
+		, PlatformData(InMetaData)
+	{}
+
 	/** Id of the user to send the invitation to */
 	FUniqueNetIdRef Id;
 	/** Additional data to provide context for the invitee */
