Cleanup in SocialToolkit

- move functions that handle results of sending and accepting friend requests to protected so that overriding functions can complete the code path

- rename InternalSendFriendInvite to SendFriendInviteInternal to follow naming convention in the code

-rename HandleFriendInviteSent to HandleSendFriendInviteComplete to match conventions

#tests full compile, sent and accepted friend requests

diff --git a/Engine/Plugins/Online/OnlineFramework/Source/Party/Private/SocialToolkit.cpp b/Engine/Plugins/Online/OnlineFramework/Source/Party/Private/SocialToolkit.cpp
--- a/Engine/Plugins/Online/OnlineFramework/Source/Party/Private/SocialToolkit.cpp
+++ b/Engine/Plugins/Online/OnlineFramework/Source/Party/Private/SocialToolkit.cpp
@@ -526,7 +526,7 @@ bool USocialToolkit::TrySendFriendInvite(USocialUser& SocialUser, ESocialSubsyst
 		const bool bIsFriendshipRestricted = IsFriendshipRestricted(SocialUser, SubsystemType);
 		if (!bIsFriendshipRestricted)
 		{
-			return InternalSendFriendInvite(SocialUser, SubsystemType, OnCompleteDelegate);
+			return SendFriendInviteInternal(SocialUser, SubsystemType, OnCompleteDelegate);
 		}
 		else
 		{
@@ -542,13 +542,24 @@ bool USocialToolkit::TrySendFriendInvite(USocialUser& SocialUser, ESocialSubsyst
 }
 
 
-bool USocialToolkit::InternalSendFriendInvite(USocialUser& SocialUser, ESocialSubsystem SubsystemType, FOnTrySendFriendInviteComplete OnCompleteDelegate) const
+bool USocialToolkit::SendFriendInviteInternal(USocialUser& SocialUser, ESocialSubsystem SubsystemType, FOnTrySendFriendInviteComplete OnCompleteDelegate) const
 {
 	IOnlineFriendsPtr FriendsInterface = Online::GetFriendsInterface(GetWorld(), USocialManager::GetSocialOssName(SubsystemType));
 	const FUniqueNetIdRepl SocialUserSubsystemId = SocialUser.GetUserId(SubsystemType);
 	if (FriendsInterface && SocialUserSubsystemId.IsValid())
 	{
-		return FriendsInterface->SendInvite(GetLocalUserNum(), *SocialUserSubsystemId, FriendListToQuery, FOnSendInviteComplete::CreateUObject(const_cast<USocialToolkit*>(this), &USocialToolkit::HandleFriendInviteSent, SubsystemType, SocialUser.GetDisplayName(), OnCompleteDelegate));
+		return FriendsInterface->SendInvite(
+			GetLocalUserNum(), 
+			*SocialUserSubsystemId, 
+			FriendListToQuery, 
+			FOnSendInviteComplete::CreateUObject(
+				const_cast<USocialToolkit*>(this), 
+				&USocialToolkit::HandleSendFriendInviteComplete, 
+				SubsystemType, 
+				SocialUser.GetDisplayName(), 
+				OnCompleteDelegate
+			)
+		);
 	}
 	else
 	{
@@ -564,7 +575,16 @@ bool USocialToolkit::AcceptFriendInvite(const USocialUser& SocialUser, ESocialSu
 	{
 		IOnlineFriendsPtr FriendsInterface = GetSocialOss(SubsystemType)->GetFriendsInterface();
 		check(FriendsInterface.IsValid());
-		return FriendsInterface->AcceptInvite(GetLocalUserNum(), *SocialUser.GetUserId(SubsystemType), EFriendsLists::ToString(EFriendsLists::Default), FOnAcceptInviteComplete::CreateUObject(const_cast<USocialToolkit*>(this), &USocialToolkit::HandleAcceptFriendInviteComplete, OnCompleteDelegate));
+		return FriendsInterface->AcceptInvite(
+			GetLocalUserNum(), 
+			*SocialUser.GetUserId(SubsystemType), 
+			EFriendsLists::ToString(EFriendsLists::Default), 
+			FOnAcceptInviteComplete::CreateUObject(
+				const_cast<USocialToolkit*>(this), 
+				&USocialToolkit::HandleAcceptFriendInviteComplete, 
+				OnCompleteDelegate
+			)
+		);
 	}
 	else
 	{
@@ -1010,7 +1030,7 @@ void USocialToolkit::HandleFriendInviteRejected(const FUniqueNetId& LocalUserId,
 	}
 }
 
-void USocialToolkit::HandleFriendInviteSent(int32 LocalUserNum, bool bWasSuccessful, const FUniqueNetId& InvitedUserId, const FString& ListName, const FString& ErrorStr, ESocialSubsystem SubsystemType, FString DisplayName, FOnTrySendFriendInviteComplete OnCompleteDelegate)
+void USocialToolkit::HandleSendFriendInviteComplete(int32 LocalUserNum, bool bWasSuccessful, const FUniqueNetId& InvitedUserId, const FString& ListName, const FString& ErrorStr, ESocialSubsystem SubsystemType, FString DisplayName, FOnTrySendFriendInviteComplete OnCompleteDelegate)
 {
 	if (bWasSuccessful)
 	{
diff --git a/Engine/Plugins/Online/OnlineFramework/Source/Party/Public/SocialToolkit.h b/Engine/Plugins/Online/OnlineFramework/Source/Party/Public/SocialToolkit.h
--- a/Engine/Plugins/Online/OnlineFramework/Source/Party/Public/SocialToolkit.h
+++ b/Engine/Plugins/Online/OnlineFramework/Source/Party/Public/SocialToolkit.h
@@ -166,7 +166,7 @@ PACKAGE_SCOPE:
 
 protected:
 	/** protected function that only handles sending a friend request via the friends interface, assumes all checks are handled previously **/
-	virtual bool InternalSendFriendInvite(USocialUser& SocialUser, ESocialSubsystem SubsystemType, FOnTrySendFriendInviteComplete OnCompleteDelegate = FOnTrySendFriendInviteComplete()) const;
+	virtual bool SendFriendInviteInternal(USocialUser& SocialUser, ESocialSubsystem SubsystemType, FOnTrySendFriendInviteComplete OnCompleteDelegate = FOnTrySendFriendInviteComplete()) const;
 
 	virtual void OnOwnerLoggedIn();
 	virtual void OnOwnerLoggedOut();
@@ -191,7 +191,12 @@ protected:
 	virtual void OnQueryFriendsListSuccess(ESocialSubsystem SubsystemType, const TArray<TSharedRef<FOnlineFriend>>& FriendsList) {}
 	virtual void OnQueryBlockedPlayersSuccess(ESocialSubsystem SubsystemType, const TArray<TSharedRef<FOnlineBlockedPlayer>>& BlockedPlayers) {}
 	virtual void OnQueryRecentPlayersSuccess(ESocialSubsystem SubsystemType, const TArray<TSharedRef<FOnlineRecentPlayer>>& FriendsList) {}
-	
+
+	/** handle result of TrySendFriendInviteInternal **/
+	void HandleSendFriendInviteComplete(int32 LocalUserNum, bool bWasSuccessful, const FUniqueNetId& InvitedUserId, const FString& ListName, const FString& ErrorStr, ESocialSubsystem SubsystemType, FString DisplayName, FOnTrySendFriendInviteComplete OnCompleteDelegate);
+	/** handle result of AcceptFriendInvite **/
+	void HandleAcceptFriendInviteComplete(int32 LocalUserNum, bool bWasSuccessful, const FUniqueNetId& InviterUserId, const FString& ListName, const FString& ErrorStr, FOnAcceptFriendInviteComplete OnCompleteDelegate);
+
 	/** The type of SocialUser to create to represent known users */
 	TSubclassOf<USocialUser> SocialUserClass;
 
@@ -251,12 +256,10 @@ private:	// Handlers
 	void HandleFriendInviteReceived(const FUniqueNetId& LocalUserId, const FUniqueNetId& SenderId, ESocialSubsystem SubsystemType);
 	void HandleFriendInviteAccepted(const FUniqueNetId& LocalUserId, const FUniqueNetId& NewFriendId, ESocialSubsystem SubsystemType);
 	void HandleFriendInviteRejected(const FUniqueNetId& LocalUserId, const FUniqueNetId& RejecterId, ESocialSubsystem SubsystemType);
-	void HandleFriendInviteSent(int32 LocalUserNum, bool bWasSuccessful, const FUniqueNetId& InvitedUserId, const FString& ListName, const FString& ErrorStr, ESocialSubsystem SubsystemType, FString DisplayName, FOnTrySendFriendInviteComplete OnCompleteDelegate);
 	void HandleFriendRemoved(const FUniqueNetId& LocalUserId, const FUniqueNetId& FormerFriendId, ESocialSubsystem SubsystemType);
 
 	void HandleDeleteFriendComplete(int32 LocalPlayer, bool bWasSuccessful, const FUniqueNetId& FormerFriendId, const FString& ListName, const FString& ErrorStr, ESocialSubsystem SubsystemType);
-	void HandleAcceptFriendInviteComplete(int32 LocalUserNum, bool bWasSuccessful, const FUniqueNetId& InviterUserId, const FString& ListName, const FString& ErrorStr, FOnAcceptFriendInviteComplete OnCompleteDelegate);
-
+	
 	void HandlePartyInviteReceived(const FUniqueNetId& LocalUserId, const IOnlinePartyJoinInfo& Invite);
 	void HandlePartyInviteRemoved(const FUniqueNetId& LocalUserId, const IOnlinePartyJoinInfo& Invite, EPartyInvitationRemovedReason Reason);
 #if PARTY_PLATFORM_INVITE_PERMISSIONS
