Undo unnecessary changes in SocialToolkit

We made a few changes that turned out to not be necessary. This change reverts the following
- moved try an accept friend invite out of package scope
- added a delegate callback to try and accept invite functions

And adds
- created AcceptFriendInviteInternal to match the SendFriendInviteInternal implementation

changes getting partially reverted
https://p4-swarm.epicgames.net/changes/19117440
https://p4-swarm.epicgames.net/changes/19131808

diff --git a/Engine/Plugins/Online/OnlineFramework/Source/Party/Private/SocialToolkit.cpp b/Engine/Plugins/Online/OnlineFramework/Source/Party/Private/SocialToolkit.cpp
--- a/Engine/Plugins/Online/OnlineFramework/Source/Party/Private/SocialToolkit.cpp
+++ b/Engine/Plugins/Online/OnlineFramework/Source/Party/Private/SocialToolkit.cpp
@@ -334,18 +334,18 @@ USocialUser* USocialToolkit::FindUser(const FUniqueNetIdRepl& UserId) const
 	return FoundUser ? FoundUser->Get() : nullptr;
 }
 
-void USocialToolkit::TrySendFriendInvite(const FString& DisplayNameOrEmail, FOnTrySendFriendInviteComplete OnCompleteDelegate) const
+void USocialToolkit::TrySendFriendInvite(const FString& DisplayNameOrEmail) const
 {
 	IOnlineSubsystem* PrimaryOSS = GetSocialOss(ESocialSubsystem::Primary);
 	IOnlineUserPtr UserInterface = PrimaryOSS ? PrimaryOSS->GetUserInterface() : nullptr;
 	if (UserInterface.IsValid())
 	{
-		IOnlineUser::FOnQueryUserMappingComplete QueryCompleteDelegate = IOnlineUser::FOnQueryUserMappingComplete::CreateUObject(const_cast<USocialToolkit*>(this), &USocialToolkit::HandleQueryPrimaryUserIdMappingComplete, OnCompleteDelegate);
+		IOnlineUser::FOnQueryUserMappingComplete QueryCompleteDelegate = IOnlineUser::FOnQueryUserMappingComplete::CreateUObject(const_cast<USocialToolkit*>(this), &USocialToolkit::HandleQueryPrimaryUserIdMappingComplete);
 		UserInterface->QueryUserIdMapping(*GetLocalUserNetId(ESocialSubsystem::Primary), DisplayNameOrEmail, QueryCompleteDelegate);
 	}
 	else
 	{
-		OnCompleteDelegate.ExecuteIfBound(false, FriendInviteFailureReason::InviteFailReason_Uninitialized);
+		UE_LOG(LogParty, Log, TEXT("SocialToolkit [%d] failed to execute TrySendFriendInvite."), GetLocalUserNum());
 	}
 }
 
@@ -513,36 +513,25 @@ void USocialToolkit::NotifySubsystemIdEstablished(USocialUser& SocialUser, ESoci
 	}
 }
 
-bool USocialToolkit::TrySendFriendInvite(USocialUser& SocialUser, ESocialSubsystem SubsystemType, FOnTrySendFriendInviteComplete OnCompleteDelegate) const
+bool USocialToolkit::TrySendFriendInvite(USocialUser& SocialUser, ESocialSubsystem SubsystemType) const
 {
 	if (SocialUser.GetFriendInviteStatus(SubsystemType) == EInviteStatus::PendingOutbound)
 	{
 		OnFriendInviteSent().Broadcast(SocialUser, SubsystemType);
-		OnCompleteDelegate.ExecuteIfBound(true, FString());
 		return true;
 	}
-	else if (!SocialUser.IsFriend(SubsystemType))
+	else if (!SocialUser.IsFriend(SubsystemType) && !IsFriendshipRestricted(SocialUser, SubsystemType))
 	{
-		const bool bIsFriendshipRestricted = IsFriendshipRestricted(SocialUser, SubsystemType);
-		if (!bIsFriendshipRestricted)
-		{
-			return SendFriendInviteInternal(SocialUser, SubsystemType, OnCompleteDelegate);
-		}
-		else
-		{
-			OnCompleteDelegate.ExecuteIfBound(false, FriendInviteFailureReason::InviteFailReason_FriendshipRestricted);
-			return false;
-		}
+		return SendFriendInviteInternal(SocialUser, SubsystemType);
 	}
 	else
 	{
-		OnCompleteDelegate.ExecuteIfBound(false, FriendInviteFailureReason::InviteFailReason_AlreadyFriends);
 		return false;
 	}
 }
 
 
-bool USocialToolkit::SendFriendInviteInternal(USocialUser& SocialUser, ESocialSubsystem SubsystemType, FOnTrySendFriendInviteComplete OnCompleteDelegate) const
+bool USocialToolkit::SendFriendInviteInternal(USocialUser& SocialUser, ESocialSubsystem SubsystemType) const
 {
 	IOnlineFriendsPtr FriendsInterface = Online::GetFriendsInterface(GetWorld(), USocialManager::GetSocialOssName(SubsystemType));
 	const FUniqueNetIdRepl SocialUserSubsystemId = SocialUser.GetUserId(SubsystemType);
@@ -556,43 +545,44 @@ bool USocialToolkit::SendFriendInviteInternal(USocialUser& SocialUser, ESocialSu
 				const_cast<USocialToolkit*>(this), 
 				&USocialToolkit::HandleSendFriendInviteComplete, 
 				SubsystemType, 
-				SocialUser.GetDisplayName(), 
-				OnCompleteDelegate
+				SocialUser.GetDisplayName()
 			)
 		);
 	}
 	else
 	{
-		OnCompleteDelegate.ExecuteIfBound(false, FriendInviteFailureReason::InviteFailReason_Uninitialized);
 		return false;
 	}
 }
 
 
-bool USocialToolkit::AcceptFriendInvite(const USocialUser& SocialUser, ESocialSubsystem SubsystemType, FOnAcceptFriendInviteComplete OnCompleteDelegate) const
+bool USocialToolkit::AcceptFriendInvite(const USocialUser& SocialUser, ESocialSubsystem SubsystemType) const
 {
 	if (SocialUser.GetFriendInviteStatus(SubsystemType) == EInviteStatus::PendingInbound)
 	{
-		IOnlineFriendsPtr FriendsInterface = GetSocialOss(SubsystemType)->GetFriendsInterface();
-		check(FriendsInterface.IsValid());
-		return FriendsInterface->AcceptInvite(
-			GetLocalUserNum(), 
-			*SocialUser.GetUserId(SubsystemType), 
-			EFriendsLists::ToString(EFriendsLists::Default), 
-			FOnAcceptInviteComplete::CreateUObject(
-				const_cast<USocialToolkit*>(this), 
-				&USocialToolkit::HandleAcceptFriendInviteComplete, 
-				OnCompleteDelegate
-			)
-		);
+		return AcceptFriendInviteInternal(SocialUser, SubsystemType);
 	}
 	else
 	{
-		OnCompleteDelegate.ExecuteIfBound(false, FriendAcceptFailureReason::AcceptFailReason_NotPendingInbound);
 		return false;
 	}
 }
 
+bool USocialToolkit::AcceptFriendInviteInternal(const USocialUser& SocialUser, ESocialSubsystem SubsystemType) const
+{
+	IOnlineFriendsPtr FriendsInterface = GetSocialOss(SubsystemType)->GetFriendsInterface();
+	check(FriendsInterface.IsValid());
+	return FriendsInterface->AcceptInvite(
+		GetLocalUserNum(),
+		*SocialUser.GetUserId(SubsystemType),
+		EFriendsLists::ToString(EFriendsLists::Default),
+		FOnAcceptInviteComplete::CreateUObject(
+			const_cast<USocialToolkit*>(this),
+			&USocialToolkit::HandleAcceptFriendInviteComplete
+		)
+	);
+}
+
 bool USocialToolkit::IsFriendshipRestricted(const USocialUser& SocialUser, ESocialSubsystem SubsystemType) const
 {
 	return false;
@@ -942,36 +932,32 @@ void USocialToolkit::HandlePresenceReceived(const FUniqueNetId& UserId, const TS
 	}
 }
 
-void USocialToolkit::HandleQueryPrimaryUserIdMappingComplete(bool bWasSuccessful, const FUniqueNetId& RequestingUserId, const FString& DisplayName, const FUniqueNetId& IdentifiedUserId, const FString& Error, FOnTrySendFriendInviteComplete OnCompleteDelegate)
+void USocialToolkit::HandleQueryPrimaryUserIdMappingComplete(bool bWasSuccessful, const FUniqueNetId& RequestingUserId, const FString& DisplayName, const FUniqueNetId& IdentifiedUserId, const FString& Error)
 {
 	if (!IdentifiedUserId.IsValid())
 	{
 		OnSendFriendInviteComplete(IdentifiedUserId, DisplayName, false, FriendInviteFailureReason::InviteFailReason_NotFound);
-		OnCompleteDelegate.ExecuteIfBound(false, FriendInviteFailureReason::InviteFailReason_NotFound);
 	}
 	else if (RequestingUserId == IdentifiedUserId)
 	{
 		OnSendFriendInviteComplete(IdentifiedUserId, DisplayName, false, FriendInviteFailureReason::InviteFailReason_AddingSelfFail);
-		OnCompleteDelegate.ExecuteIfBound(false, FriendInviteFailureReason::InviteFailReason_AddingSelfFail);
 	}
 	else
 	{
 		QueueUserDependentActionInternal(IdentifiedUserId.AsShared(), ESocialSubsystem::Primary,
-			[this, DisplayName, OnCompleteDelegate] (USocialUser& SocialUser)
+			[this, DisplayName] (USocialUser& SocialUser)
 			{
 				if (SocialUser.IsBlocked())
 				{
 					OnSendFriendInviteComplete(*SocialUser.GetUserId(ESocialSubsystem::Primary), DisplayName, false, FriendInviteFailureReason::InviteFailReason_AddingBlockedFail);
-					OnCompleteDelegate.ExecuteIfBound(false, FriendInviteFailureReason::InviteFailReason_AddingBlockedFail);
 				}
 				else if (SocialUser.IsFriend(ESocialSubsystem::Primary))
 				{
 					OnSendFriendInviteComplete(*SocialUser.GetUserId(ESocialSubsystem::Primary), DisplayName, false, FriendInviteFailureReason::InviteFailReason_AlreadyFriends);
-					OnCompleteDelegate.ExecuteIfBound(false, FriendInviteFailureReason::InviteFailReason_AlreadyFriends);
 				}
 				else
 				{
-					TrySendFriendInvite(SocialUser, ESocialSubsystem::Primary, OnCompleteDelegate);
+					TrySendFriendInvite(SocialUser, ESocialSubsystem::Primary);
 				}
 			});
 	}
@@ -1030,7 +1016,7 @@ void USocialToolkit::HandleFriendInviteRejected(const FUniqueNetId& LocalUserId,
 	}
 }
 
-void USocialToolkit::HandleSendFriendInviteComplete(int32 LocalUserNum, bool bWasSuccessful, const FUniqueNetId& InvitedUserId, const FString& ListName, const FString& ErrorStr, ESocialSubsystem SubsystemType, FString DisplayName, FOnTrySendFriendInviteComplete OnCompleteDelegate)
+void USocialToolkit::HandleSendFriendInviteComplete(int32 LocalUserNum, bool bWasSuccessful, const FUniqueNetId& InvitedUserId, const FString& ListName, const FString& ErrorStr, ESocialSubsystem SubsystemType, FString DisplayName)
 {
 	if (bWasSuccessful)
 	{
@@ -1049,7 +1035,6 @@ void USocialToolkit::HandleSendFriendInviteComplete(int32 LocalUserNum, bool bWa
 			});
 	}
 	OnSendFriendInviteComplete(InvitedUserId, DisplayName, bWasSuccessful, ErrorStr);
-	OnCompleteDelegate.ExecuteIfBound(bWasSuccessful, ErrorStr);
 }
 
 void USocialToolkit::HandleFriendRemoved(const FUniqueNetId& LocalUserId, const FUniqueNetId& FriendId, ESocialSubsystem SubsystemType)
@@ -1078,10 +1063,9 @@ void USocialToolkit::HandleDeleteFriendComplete(int32 InLocalUserNum, bool bWasS
 	OnDeleteFriendComplete(InLocalUserNum, bWasSuccessful, DeletedFriendId, ListName, ErrorStr, SubsystemType);
 }
 
-void USocialToolkit::HandleAcceptFriendInviteComplete(int32 LocalUserNum, bool bWasSuccessful, const FUniqueNetId& InviterUserId, const FString& ListName, const FString& ErrorStr, FOnAcceptFriendInviteComplete OnCompleteDelegate)
+void USocialToolkit::HandleAcceptFriendInviteComplete(int32 LocalUserNum, bool bWasSuccessful, const FUniqueNetId& InviterUserId, const FString& ListName, const FString& ErrorStr)
 {
 	OnAcceptFriendInviteComplete(InviterUserId, bWasSuccessful, ErrorStr);
-	OnCompleteDelegate.ExecuteIfBound(bWasSuccessful, ErrorStr);
 }
 
 const bool USocialToolkit::IsInviteAllowedFromUser(const USocialUser& User, const TSharedRef<const IOnlinePartyJoinInfo>& InvitePtr) const
diff --git a/Engine/Plugins/Online/OnlineFramework/Source/Party/Public/SocialToolkit.h b/Engine/Plugins/Online/OnlineFramework/Source/Party/Public/SocialToolkit.h
--- a/Engine/Plugins/Online/OnlineFramework/Source/Party/Public/SocialToolkit.h
+++ b/Engine/Plugins/Online/OnlineFramework/Source/Party/Public/SocialToolkit.h
@@ -30,8 +30,6 @@ namespace FriendInviteFailureReason
 	const FString InviteFailReason_AddingSelfFail = TEXT("AddingSelfFail");
 	const FString InviteFailReason_AddingBlockedFail = TEXT("AddingBlockedFail");
 	const FString InviteFailReason_AlreadyFriends = TEXT("AlreadyFriends");
-	const FString InviteFailReason_Uninitialized = TEXT("Uninitialized");
-	const FString InviteFailReason_FriendshipRestricted = TEXT("FriendshipRestricted");
 }
 namespace FriendAcceptFailureReason
 {
@@ -39,8 +37,6 @@ namespace FriendAcceptFailureReason
 }
 
 DECLARE_DELEGATE_OneParam(FUserDependentAction, USocialUser&);
-DECLARE_DELEGATE_TwoParams(FOnTrySendFriendInviteComplete, bool /* bWasSuccessful */, FString /* ErrorString */);
-DECLARE_DELEGATE_TwoParams(FOnAcceptFriendInviteComplete, bool /* bWasSuccessful */, FString /* ErrorString */);
 
 /** Represents the full suite of social functionality available to a given LocalPlayer */
 UCLASS(Within = SocialManager)
@@ -101,12 +97,7 @@ public:
 	 * Attempts to send a friend invite to another user based on display name or email.
 	 * Only necessary to use this path when you do not have a known USocialUser for this user already.
 	 */
-	virtual void TrySendFriendInvite(const FString& DisplayNameOrEmail, FOnTrySendFriendInviteComplete OnCompleteDelegate = FOnTrySendFriendInviteComplete()) const;
-
-	/** Send friend invite to a SocialUser */
-	virtual bool TrySendFriendInvite(USocialUser& SocialUser, ESocialSubsystem SubsystemType, FOnTrySendFriendInviteComplete OnCompleteDelegate = FOnTrySendFriendInviteComplete()) const;
-
-	virtual bool AcceptFriendInvite(const USocialUser& SocialUser, ESocialSubsystem SubsystemType, FOnAcceptFriendInviteComplete OnCompleteDelegate = FOnAcceptFriendInviteComplete()) const;
+	void TrySendFriendInvite(const FString& DisplayNameOrEmail) const;
 
 	virtual bool IsFriendshipRestricted(const USocialUser& SocialUser, ESocialSubsystem SubsystemType) const;
 
@@ -155,6 +146,10 @@ public:
 PACKAGE_SCOPE:
 	void NotifySubsystemIdEstablished(USocialUser& SocialUser, ESocialSubsystem SubsystemType, const FUniqueNetIdRepl& SubsystemId);
 	TSubclassOf<USocialChatManager> GetChatManagerClass() { return ChatManagerClass; }
+	
+	bool TrySendFriendInvite(USocialUser& SocialUser, ESocialSubsystem SubsystemType) const;
+
+	bool AcceptFriendInvite(const USocialUser& SocialUser, ESocialSubsystem SubsystemType) const;
 
 	void RequestToJoinParty(USocialUser& SocialUser);
 
@@ -165,8 +160,10 @@ PACKAGE_SCOPE:
 #endif
 
 protected:
-	/** protected function that only handles sending a friend request via the friends interface, assumes all checks are handled previously **/
-	virtual bool SendFriendInviteInternal(USocialUser& SocialUser, ESocialSubsystem SubsystemType, FOnTrySendFriendInviteComplete OnCompleteDelegate = FOnTrySendFriendInviteComplete()) const;
+	/** only handles sending a friend request via the friends interface, assumes all checks are handled previously **/
+	virtual bool SendFriendInviteInternal(USocialUser& SocialUser, ESocialSubsystem SubsystemType) const;
+	/** only handles accepting a friend request via the friends interface, assumes all checks are handled previously **/
+	virtual bool AcceptFriendInviteInternal(const USocialUser& SocialUser, ESocialSubsystem SubsystemType) const;
 
 	virtual void OnOwnerLoggedIn();
 	virtual void OnOwnerLoggedOut();
@@ -193,9 +190,9 @@ protected:
 	virtual void OnQueryRecentPlayersSuccess(ESocialSubsystem SubsystemType, const TArray<TSharedRef<FOnlineRecentPlayer>>& FriendsList) {}
 
 	/** handle result of TrySendFriendInviteInternal **/
-	void HandleSendFriendInviteComplete(int32 LocalUserNum, bool bWasSuccessful, const FUniqueNetId& InvitedUserId, const FString& ListName, const FString& ErrorStr, ESocialSubsystem SubsystemType, FString DisplayName, FOnTrySendFriendInviteComplete OnCompleteDelegate);
+	void HandleSendFriendInviteComplete(int32 LocalUserNum, bool bWasSuccessful, const FUniqueNetId& InvitedUserId, const FString& ListName, const FString& ErrorStr, ESocialSubsystem SubsystemType, FString DisplayName);
 	/** handle result of AcceptFriendInvite **/
-	void HandleAcceptFriendInviteComplete(int32 LocalUserNum, bool bWasSuccessful, const FUniqueNetId& InviterUserId, const FString& ListName, const FString& ErrorStr, FOnAcceptFriendInviteComplete OnCompleteDelegate);
+	void HandleAcceptFriendInviteComplete(int32 LocalUserNum, bool bWasSuccessful, const FUniqueNetId& InviterUserId, const FString& ListName, const FString& ErrorStr);
 
 	/** The type of SocialUser to create to represent known users */
 	TSubclassOf<USocialUser> SocialUserClass;
@@ -269,7 +266,7 @@ private:	// Handlers
 	void HandleBlockPlayerComplete(int32 LocalUserNum, bool bWasSuccessful, const FUniqueNetId& BlockedPlayerID, const FString& ListName, const FString& ErrorStr, ESocialSubsystem SubsystemType);
 	void HandleUnblockPlayerComplete(int32 LocalUserNum, bool bWasSuccessful, const FUniqueNetId& UnblockedPlayerID, const FString& ListName, const FString& ErrorStr, ESocialSubsystem SubsystemType);
 	
-	void HandleQueryPrimaryUserIdMappingComplete(bool bWasSuccessful, const FUniqueNetId& RequestingUserId, const FString& DisplayName, const FUniqueNetId& IdentifiedUserId, const FString& Error, FOnTrySendFriendInviteComplete OnCompleteDelegate);
+	void HandleQueryPrimaryUserIdMappingComplete(bool bWasSuccessful, const FUniqueNetId& RequestingUserId, const FString& DisplayName, const FUniqueNetId& IdentifiedUserId, const FString& Error);
 
 	void HandlePartyMemberExited(const FUniqueNetId& LocalUserId, const FOnlinePartyId& PartyId, const FUniqueNetId& MemberId, const EMemberExitedReason Reason);
 	void HandleGameDestroyed(const FName SessionName, bool bWasSuccessful);
