(Social/Party Signal) Do not enable Join/Request To Join social menu entries for non-friends, they should only be allowed to join via the Party Signal device.

diff --git a/Engine/Plugins/Online/OnlineFramework/Source/Party/Private/Interactions/PartyInteractions.cpp b/Engine/Plugins/Online/OnlineFramework/Source/Party/Private/Interactions/PartyInteractions.cpp
--- a/Engine/Plugins/Online/OnlineFramework/Source/Party/Private/Interactions/PartyInteractions.cpp
+++ b/Engine/Plugins/Online/OnlineFramework/Source/Party/Private/Interactions/PartyInteractions.cpp
@@ -52,7 +52,7 @@ FString FSocialInteraction_JoinParty::GetSlashCommandToken()
 bool FSocialInteraction_JoinParty::CanExecute(const USocialUser& User)
 {
 	FJoinPartyResult MockJoinResult = User.CheckPartyJoinability(IOnlinePartySystem::GetPrimaryPartyTypeId());
-	return MockJoinResult.WasSuccessful();
+	return MockJoinResult.WasSuccessful() && User.IsFriend();
 }
 
 void FSocialInteraction_JoinParty::ExecuteInteraction(USocialUser& User)
@@ -76,7 +76,7 @@ FString FSocialInteraction_RequestToJoinParty::GetSlashCommandToken()
 
 bool FSocialInteraction_RequestToJoinParty::CanExecute(const USocialUser& User)
 {
-	return User.CanRequestToJoin();
+	return User.CanRequestToJoin() && User.IsFriend();
 }
 
 void FSocialInteraction_RequestToJoinParty::ExecuteInteraction(USocialUser& User)
