Remove PartyPackage.h
The mix of public/protected does not work intuitively with extensions to this plugin's classes.  Use friend to add access that PACKAGE_SCOPE was used for.

diff --git a/Engine/Plugins/Online/OnlineFramework/Source/Party/Party.Build.cs b/Engine/Plugins/Online/OnlineFramework/Source/Party/Party.Build.cs
--- a/Engine/Plugins/Online/OnlineFramework/Source/Party/Party.Build.cs
+++ b/Engine/Plugins/Online/OnlineFramework/Source/Party/Party.Build.cs
@@ -6,8 +6,6 @@ public class Party : ModuleRules
 {
 	public Party(ReadOnlyTargetRules Target) : base(Target)
 	{
-		PrivateDefinitions.Add("PARTY_PACKAGE=1");
-
 		PublicDefinitions.Add("PARTY_PLATFORM_SESSIONS_PSN=" + (bUsesPSNSessions ? "1" : "0"));
 		PublicDefinitions.Add("PARTY_PLATFORM_SESSIONS_XBL=" + (bUsesXBLSessions ? "1" : "0"));
 		PublicDefinitions.Add("PARTY_PLATFORM_INVITE_PERMISSIONS=" + (bUsesPlatformInvitePermissions? "1" : "0"));
diff --git a/Engine/Plugins/Online/OnlineFramework/Source/Party/Private/User/SocialUserList.h b/Engine/Plugins/Online/OnlineFramework/Source/Party/Private/User/SocialUserList.h
--- a/Engine/Plugins/Online/OnlineFramework/Source/Party/Private/User/SocialUserList.h
+++ b/Engine/Plugins/Online/OnlineFramework/Source/Party/Private/User/SocialUserList.h
@@ -4,7 +4,6 @@
 
 #include "User/ISocialUserList.h"
 #include "UObject/GCObject.h"
-#include "PartyPackage.h"
 #include "Containers/Ticker.h"
 
 enum class EMemberExitedReason : uint8;
@@ -12,6 +11,7 @@ typedef TSharedRef<const IOnlinePartyRequestToJoinInfo> IOnlinePartyRequestToJoi
 
 class FSocialUserList : public ISocialUserList, public FGCObject, public TSharedFromThis<FSocialUserList>
 {
+	friend class USocialToolkit;
 public:
 	static TSharedRef<FSocialUserList> CreateUserList(const USocialToolkit& InOwnerToolkit, const FSocialUserListConfig& Config);
 
@@ -34,7 +34,7 @@ public:
 
 	bool HasPresenceFilters() const;
 
-PACKAGE_SCOPE:
+protected:
 	const FSocialUserListConfig& GetListConfig() const { return ListConfig; }
 
 private:
diff --git a/Engine/Plugins/Online/OnlineFramework/Source/Party/Public/Party/PartyDataReplicator.h b/Engine/Plugins/Online/OnlineFramework/Source/Party/Public/Party/PartyDataReplicator.h
--- a/Engine/Plugins/Online/OnlineFramework/Source/Party/Public/Party/PartyDataReplicator.h
+++ b/Engine/Plugins/Online/OnlineFramework/Source/Party/Public/Party/PartyDataReplicator.h
@@ -7,20 +7,20 @@
 #include "Containers/Ticker.h"
 #include "Interfaces/OnlinePartyInterface.h"
 #include "Stats/Stats.h"
-#include "PartyPackage.h"
 
 /** Util exclusively for use by TPartyDataReplicator to circumvent circular include header issues (we can't include SocialParty.h or PartyMember.h here) */
 class FPartyDataReplicatorHelper
 {
-	template <typename> friend class TPartyDataReplicator;
+	template <typename, class> friend class TPartyDataReplicator;
 	PARTY_API static void ReplicateDataToMembers(const FOnlinePartyRepDataBase& RepDataInstance, const UScriptStruct& RepDataType, const FOnlinePartyData& ReplicationPayload);
 };
 
 /** Base util class for dealing with data that is replicated to party members */
-template <typename RepDataT>
+template <typename RepDataT, class OwningObjectT>
 class TPartyDataReplicator : public FGCObject
 {
 	static_assert(TIsDerivedFrom<RepDataT, FOnlinePartyRepDataBase>::IsDerived, "TPartyDataReplicator is only intended to function with FOnlinePartyRepDataBase types.");
+	friend OwningObjectT;
 
 public:
 	~TPartyDataReplicator()
@@ -71,7 +71,7 @@ public:
 		}
 	}
 
-PACKAGE_SCOPE:
+protected:
 	void ProcessReceivedData(const FOnlinePartyData& IncomingPartyData, bool bCompareToPrevious = true)
 	{
 		// If the rep data can be edited locally, disregard any replication updates (they're the same at best or out of date at worst)
diff --git a/Engine/Plugins/Online/OnlineFramework/Source/Party/Public/Party/PartyMember.h b/Engine/Plugins/Online/OnlineFramework/Source/Party/Public/Party/PartyMember.h
--- a/Engine/Plugins/Online/OnlineFramework/Source/Party/Public/Party/PartyMember.h
+++ b/Engine/Plugins/Online/OnlineFramework/Source/Party/Public/Party/PartyMember.h
@@ -5,7 +5,6 @@
 #include "Party/PartyTypes.h"
 #include "Party/PartyDataReplicator.h"
 
-#include "PartyPackage.h"
 #include "PartyMember.generated.h"
 
 class USocialUser;
@@ -72,13 +71,15 @@ private:
 	EXPOSE_REP_DATA_PROPERTY(FPartyMemberRepData, ECrossplayPreference, CrossplayPreference);
 };
 
-using FPartyMemberDataReplicator = TPartyDataReplicator<FPartyMemberRepData>;
+using FPartyMemberDataReplicator = TPartyDataReplicator<FPartyMemberRepData, UPartyMember>;
 
 UCLASS(Abstract, config = Game, Within = SocialParty, Transient)
 class PARTY_API UPartyMember : public UObject
 {
 	GENERATED_BODY()
 
+	friend class FPartyPlatformSessionMonitor;
+	friend USocialParty;
 public:
 	UPartyMember();
 
@@ -114,7 +115,7 @@ public:
 
 	FString ToDebugString(bool bIncludePartyId = true) const;
 
-PACKAGE_SCOPE:
+protected:
 	void InitializePartyMember(const FOnlinePartyMemberConstRef& OssMember, const FSimpleDelegate& OnInitComplete);
 
 	FPartyMemberRepData& GetMutableRepData() { return *MemberDataReplicator; }
diff --git a/Engine/Plugins/Online/OnlineFramework/Source/Party/Public/Party/PartyTypes.h b/Engine/Plugins/Online/OnlineFramework/Source/Party/Public/Party/PartyTypes.h
--- a/Engine/Plugins/Online/OnlineFramework/Source/Party/Public/Party/PartyTypes.h
+++ b/Engine/Plugins/Online/OnlineFramework/Source/Party/Public/Party/PartyTypes.h
@@ -461,7 +461,7 @@ protected:
 	void LogPropertyChanged(const TCHAR* OwningStructTypeName, const TCHAR* PropertyName, bool bFromReplication) const;
 
 	friend class FPartyDataReplicatorHelper;
-	template <typename> friend class TPartyDataReplicator;
+	template <typename, class> friend class TPartyDataReplicator;
 	FSimpleDelegate OnDataChanged;
 };
 
diff --git a/Engine/Plugins/Online/OnlineFramework/Source/Party/Public/Party/SocialParty.h b/Engine/Plugins/Online/OnlineFramework/Source/Party/Public/Party/SocialParty.h
--- a/Engine/Plugins/Online/OnlineFramework/Source/Party/Public/Party/SocialParty.h
+++ b/Engine/Plugins/Online/OnlineFramework/Source/Party/Public/Party/SocialParty.h
@@ -12,7 +12,6 @@
 #include "Interfaces/OnlineChatInterface.h"
 #include "Containers/Queue.h"
 #include "Engine/EngineBaseTypes.h"
-#include "PartyPackage.h"
 #include "SocialParty.generated.h"
 
 class APartyBeaconClient;
@@ -64,7 +63,7 @@ private:
 	mutable FSimpleMulticastDelegate OnPlatformSessionsChangedEvent;
 };
 
-using FPartyDataReplicator = TPartyDataReplicator<FPartyRepData>;
+using FPartyDataReplicator = TPartyDataReplicator<FPartyRepData, USocialParty>;
 
 /**
  * Party game state that contains all information relevant to the communication within a party
@@ -75,6 +74,10 @@ class PARTY_API USocialParty : public UObject
 {
 	GENERATED_BODY()
 
+	friend class FPartyPlatformSessionMonitor;
+	friend UPartyMember;
+	friend USocialManager;
+	friend USocialUser;
 public:
 	static bool IsJoiningDuringLoadEnabled();
 
@@ -207,14 +210,14 @@ public:
 	virtual bool ShouldAlwaysJoinPlatformSession(const FSessionId& SessionId) const;
 
 	virtual void JoinSessionCompleteAnalytics(const FSessionId& SessionId, const FString& JoinBootableGroupSessionResult);
+	bool IsCurrentlyLeaving() const;
 
-PACKAGE_SCOPE:
+protected:
 	void InitializeParty(const TSharedRef<const FOnlineParty>& InOssParty);
 	bool IsInitialized() const;
 	void TryFinishInitialization();
 
 	bool ShouldCacheForRejoinOnDisconnect() const;
-	bool IsCurrentlyLeaving() const;
 
 	void SetIsMissingPlatformSession(bool bInIsMissingPlatformSession);
 	bool IsMissingPlatformSession() { return bIsMissingPlatformSession; }
@@ -222,7 +225,7 @@ PACKAGE_SCOPE:
 	FPartyRepData& GetMutableRepData() { return *PartyDataReplicator; }
 
 	//--------------------------
-	// User/member-specific actions that are best exposed on the individuals themselves, but best handled by the actual party (thus the package scoping)
+	// User/member-specific actions that are best exposed on the individuals themselves, but best handled by the actual party
 	bool HasUserBeenInvited(const USocialUser& User) const;
 	
 	bool CanInviteUser(const USocialUser& User) const;
diff --git a/Engine/Plugins/Online/OnlineFramework/Source/Party/Public/PartyModule.h b/Engine/Plugins/Online/OnlineFramework/Source/Party/Public/PartyModule.h
--- a/Engine/Plugins/Online/OnlineFramework/Source/Party/Public/PartyModule.h
+++ b/Engine/Plugins/Online/OnlineFramework/Source/Party/Public/PartyModule.h
@@ -7,7 +7,6 @@
 #include "Misc/CoreMisc.h"
 #include "Modules/ModuleInterface.h"
 #include "Modules/ModuleManager.h"
-#include "PartyPackage.h"
 
 /** Logging related to parties */
 PARTY_API DECLARE_LOG_CATEGORY_EXTERN(LogParty, Display, All);
diff --git a/Engine/Plugins/Online/OnlineFramework/Source/Party/Public/PartyPackage.h b/Engine/Plugins/Online/OnlineFramework/Source/Party/Public/PartyPackage.h
deleted file mode 100644
diff --git a/Engine/Plugins/Online/OnlineFramework/Source/Party/Public/SocialManager.h b/Engine/Plugins/Online/OnlineFramework/Source/Party/Public/SocialManager.h
--- a/Engine/Plugins/Online/OnlineFramework/Source/Party/Public/SocialManager.h
+++ b/Engine/Plugins/Online/OnlineFramework/Source/Party/Public/SocialManager.h
@@ -11,7 +11,6 @@
 #include "Interfaces/OnlinePartyInterface.h"
 #include "Interactions/SocialInteractionHandle.h"
 
-#include "PartyPackage.h"
 #include "SocialManager.generated.h"
 
 class ULocalPlayer;
@@ -34,6 +33,10 @@ class PARTY_API USocialManager : public UObject, public FExec
 {
 	GENERATED_BODY()
 
+	friend class FPartyPlatformSessionManager;
+	friend UPartyMember;
+	friend USocialUser;
+
 public:
 	// FExec
 	virtual bool Exec(class UWorld* InWorld, const TCHAR* Cmd, FOutputDevice& Out) override;
@@ -111,8 +114,7 @@ public:
 	/** Validates that the target user has valid join info for us to use and that we can join any party of the given type */
 	virtual FJoinPartyResult ValidateJoinTarget(const USocialUser& UserToJoin, const FOnlinePartyTypeId& PartyTypeId) const;
 
-PACKAGE_SCOPE:
-	
+protected:
 	DECLARE_DELEGATE_OneParam(FOnJoinPartyAttemptComplete, const FJoinPartyResult&);
 	void JoinParty(const USocialUser& UserToJoin, const FOnlinePartyTypeId& PartyTypeId, const FOnJoinPartyAttemptComplete& OnJoinPartyComplete);
 
diff --git a/Engine/Plugins/Online/OnlineFramework/Source/Party/Public/SocialToolkit.h b/Engine/Plugins/Online/OnlineFramework/Source/Party/Public/SocialToolkit.h
--- a/Engine/Plugins/Online/OnlineFramework/Source/Party/Public/SocialToolkit.h
+++ b/Engine/Plugins/Online/OnlineFramework/Source/Party/Public/SocialToolkit.h
@@ -10,7 +10,6 @@
 #include "Interfaces/OnlinePresenceInterface.h"
 #include "Containers/Ticker.h"
 
-#include "PartyPackage.h"
 #include "SocialToolkit.generated.h"
 
 class ULocalPlayer;
@@ -44,6 +43,10 @@ class PARTY_API USocialToolkit : public UObject, public FExec
 {
 	GENERATED_BODY()
 
+	friend class FPartyPlatformSessionManager;
+	friend USocialChatManager;
+	friend USocialUser;
+
 public:
 	template <typename ToolkitT = USocialToolkit>
 	static ToolkitT* GetToolkitForPlayer(const ULocalPlayer* LocalPlayer)
@@ -139,11 +142,13 @@ public:
 	virtual void NotifyPartyInviteReceived(USocialUser& SocialUser, const IOnlinePartyJoinInfo& Invite);
 	virtual void NotifyPartyInviteRemoved(USocialUser& SocialUser, const IOnlinePartyJoinInfo& Invite);
 
+	void RequestToJoinParty(USocialUser& SocialUser);
+
 #if WITH_EDITOR
 	bool Debug_IsRandomlyChangingPresence() const { return bDebug_IsRandomlyChangingUserPresence; }
 #endif
 
-PACKAGE_SCOPE:
+protected:
 	void NotifySubsystemIdEstablished(USocialUser& SocialUser, ESocialSubsystem SubsystemType, const FUniqueNetIdRepl& SubsystemId);
 	TSubclassOf<USocialChatManager> GetChatManagerClass() { return ChatManagerClass; }
 	
@@ -151,8 +156,6 @@ PACKAGE_SCOPE:
 
 	bool AcceptFriendInvite(const USocialUser& SocialUser, ESocialSubsystem SubsystemType) const;
 
-	void RequestToJoinParty(USocialUser& SocialUser);
-
 	void HandleUserInvalidated(USocialUser& InvalidUser);
 
 #if PARTY_PLATFORM_SESSIONS_PSN
diff --git a/Engine/Plugins/Online/OnlineFramework/Source/Party/Public/User/SocialUser.h b/Engine/Plugins/Online/OnlineFramework/Source/Party/Public/User/SocialUser.h
--- a/Engine/Plugins/Online/OnlineFramework/Source/Party/Public/User/SocialUser.h
+++ b/Engine/Plugins/Online/OnlineFramework/Source/Party/Public/User/SocialUser.h
@@ -8,7 +8,6 @@
 #include "Party/PartyTypes.h"
 #include "Interactions/SocialInteractionHandle.h"
 
-#include "PartyPackage.h"
 #include "SocialUser.generated.h"
 
 class IOnlinePartyJoinInfo;
@@ -32,6 +31,8 @@ class PARTY_API USocialUser : public UObject
 {
 	GENERATED_BODY()
 
+	friend USocialToolkit;
+
 public:
 	USocialUser();
 
@@ -164,7 +165,7 @@ public:
 	void EstablishOssInfo(const TSharedRef<FOnlineBlockedPlayer>& BlockedPlayerInfo, ESocialSubsystem SubsystemType);
 	void EstablishOssInfo(const TSharedRef<FOnlineRecentPlayer>& RecentPlayerInfo, ESocialSubsystem SubsystemType);
 
-PACKAGE_SCOPE:
+protected:
 	void InitLocalUser();
 	void Initialize(const FUniqueNetIdRepl& PrimaryId);
 
