Add CVar to easily test "disconnected to party service" state.

diff --git a/Engine/Plugins/Online/OnlineFramework/Source/Party/Private/SocialManager.cpp b/Engine/Plugins/Online/OnlineFramework/Source/Party/Private/SocialManager.cpp
--- a/Engine/Plugins/Online/OnlineFramework/Source/Party/Private/SocialManager.cpp
+++ b/Engine/Plugins/Online/OnlineFramework/Source/Party/Private/SocialManager.cpp
@@ -21,6 +21,12 @@
 #include "OnlineSubsystemSessionSettings.h"
 #include "Misc/Base64.h"
 
+static TAutoConsoleVariable<bool> CVarForceDisconnectedToPartyService(
+	TEXT("SocialUI.ForceDisconnectedToPartyService"),
+	false,
+	TEXT("Overrides the bIsConnectedToPartyService state"));
+
+
 //////////////////////////////////////////////////////////////////////////
 // FRejoinableParty
 //////////////////////////////////////////////////////////////////////////
@@ -1236,6 +1242,11 @@ void USocialManager::HandleLeavePartyForJoinComplete(ELeavePartyCompletionResult
 	}
 }
 
+bool USocialManager::IsConnectedToPartyService() const
+{
+	return bIsConnectedToPartyService && !CVarForceDisconnectedToPartyService.GetValueOnAnyThread();
+}
+
 void USocialManager::HandlePartyDisconnected(USocialParty* DisconnectingParty)
 {
 	ABORT_DURING_SHUTDOWN();
diff --git a/Engine/Plugins/Online/OnlineFramework/Source/Party/Public/SocialManager.h b/Engine/Plugins/Online/OnlineFramework/Source/Party/Public/SocialManager.h
--- a/Engine/Plugins/Online/OnlineFramework/Source/Party/Public/SocialManager.h
+++ b/Engine/Plugins/Online/OnlineFramework/Source/Party/Public/SocialManager.h
@@ -98,7 +98,7 @@ public:
 		return Cast<PartyT>(GetPartyInternal(PartyId));
 	}
 
-	bool IsConnectedToPartyService() const { return bIsConnectedToPartyService; }
+	bool IsConnectedToPartyService() const;
 
 	void HandlePartyDisconnected(USocialParty* LeavingParty);
 
