- Fix CompoundSubsystemEnum ignoring ESocialSubsystem::Primary, and not creating a unique value for a given combination of subsystems.

diff --git a/Engine/Plugins/Online/OnlineFramework/Source/Party/Public/User/ISocialUserList.h b/Engine/Plugins/Online/OnlineFramework/Source/Party/Public/User/ISocialUserList.h
--- a/Engine/Plugins/Online/OnlineFramework/Source/Party/Public/User/ISocialUserList.h
+++ b/Engine/Plugins/Online/OnlineFramework/Source/Party/Public/User/ISocialUserList.h
@@ -70,7 +70,7 @@ private:
 		uint8 OutputValue = 0;
 		for (ESocialSubsystem Subsystem : Subsystems)
 		{
-			OutputValue |= (uint8)Subsystem;
+			OutputValue |= 1 << (uint8)Subsystem;
 		}
 		return OutputValue;
 	}
