(Engine/Social)  Report Rich Presence text for the local social user.
#OKFORGITHUB PUBLIC

diff --git a/Engine/Plugins/Online/OnlineFramework/Source/Party/Private/User/SocialUser.cpp b/Engine/Plugins/Online/OnlineFramework/Source/Party/Private/User/SocialUser.cpp
--- a/Engine/Plugins/Online/OnlineFramework/Source/Party/Private/User/SocialUser.cpp
+++ b/Engine/Plugins/Online/OnlineFramework/Source/Party/Private/User/SocialUser.cpp
@@ -730,7 +730,7 @@ void USocialUser::GetRichPresenceText(FText& OutRichPresence) const
 	{
 		OutRichPresence = LOCTEXT("UserStatus_Blocked", "Blocked");
 	}
-	else if (IsFriend())
+	else if (IsFriend() || IsLocalUser())
 	{
 		const FOnlineUserPresence* PrimaryPresence = GetFriendPresenceInfo(ESocialSubsystem::Primary);
 		if (PrimaryPresence && !PrimaryPresence->Status.StatusStr.IsEmpty())
