Add a Compare method to party rep data that can compare a specific value without needing a full party rep data wrapper

diff --git a/Engine/Plugins/Online/OnlineFramework/Source/Party/Public/Party/PartyTypes.h b/Engine/Plugins/Online/OnlineFramework/Source/Party/Public/Party/PartyTypes.h
--- a/Engine/Plugins/Online/OnlineFramework/Source/Party/Public/Party/PartyTypes.h
+++ b/Engine/Plugins/Online/OnlineFramework/Source/Party/Public/Party/PartyTypes.h
@@ -323,7 +323,7 @@ public:
 };
 
 
-/** Companion to EPartyJoinDisapproval to lessen the hassle of working with a "customized" enum */
+/** Companion to EPartyJoinDenialReason to lessen the hassle of working with a "customized" enum */
 struct PARTY_API FPartyJoinDenialReason
 {
 public:
@@ -375,7 +375,7 @@ private:
 	EPartyJoinDenialReason DenialReason = EPartyJoinDenialReason::NoReason;
 };
 
-struct PARTY_API FPartyJoinApproval
+struct FPartyJoinApproval
 {
 	FPartyJoinApproval() {}
 	
@@ -490,10 +490,14 @@ public:	\
 private:	\
 	void Compare##PropertyName(const Owner& OldData) const	\
 	{	\
-		if (PropertyAccess != OldData.PropertyAccess)	\
+		Compare##PropertyName(OldData.PropertyAccess); \
+	}	\
+	void Compare##PropertyName(const typename TRemoveConst<PropertyType>::Type& OldData) const	\
+	{	\
+		if (PropertyAccess != OldData)	\
 		{	\
 			LogPropertyChanged(TEXT(#Owner), TEXT(#PropertyName), true);	\
-			On##PropertyName##ChangedDif().Broadcast(PropertyAccess, OldData.PropertyAccess);	\
+			On##PropertyName##ChangedDif().Broadcast(PropertyAccess, OldData);	\
 			On##PropertyName##Changed().Broadcast(PropertyAccess);	\
 		}	\
 	}	\
