- Log bWasSuccessful in social toolkit query callbacks, as some OSS impls might return error messages on success or no error message on failure.

diff --git a/Engine/Plugins/Online/OnlineFramework/Source/Party/Private/SocialToolkit.cpp b/Engine/Plugins/Online/OnlineFramework/Source/Party/Private/SocialToolkit.cpp
--- a/Engine/Plugins/Online/OnlineFramework/Source/Party/Private/SocialToolkit.cpp
+++ b/Engine/Plugins/Online/OnlineFramework/Source/Party/Private/SocialToolkit.cpp
@@ -123,7 +123,7 @@ public:
 
 	void HandleQueryExternalIdMappingsComplete(bool bWasSuccessful, const FUniqueNetId&, const FExternalIdQueryOptions& QueryOptions, const TArray<FString>& ExternalIds, const FString& ErrorStr)
 	{
-		UE_LOG(LogParty, Log, TEXT("FSocialQuery_MapExternalIds completed query for [%d] users on subsystem [%s] with error [%s]"), ExternalIds.Num(), ToString(SubsystemType), *ErrorStr);
+		UE_LOG(LogParty, Log, TEXT("FSocialQuery_MapExternalIds completed query for [%d] users Subsystem=[%s] bWasSuccessful=[%s] Error=[%s]"), ExternalIds.Num(), ToString(SubsystemType), *LexToString(bWasSuccessful), *ErrorStr);
 
 		if (bWasSuccessful)
 		{
@@ -804,7 +804,7 @@ void USocialToolkit::HandlePlayerLoginStatusChanged(int32 LocalUserNum, ELoginSt
 
 void USocialToolkit::HandleReadFriendsListComplete(int32 LocalUserNum, bool bWasSuccessful, const FString& ListName, const FString& ErrorStr, ESocialSubsystem SubsystemType)
 {
-	UE_LOG(LogParty, Log, TEXT("SocialToolkit [%d] finished querying friends list [%s] on subsystem [%s] with error [%s]."), GetLocalUserNum(), *ListName, ToString(SubsystemType), *ErrorStr);
+	UE_LOG(LogParty, Log, TEXT("SocialToolkit [%d] finished querying friends list ListName=[%s] Subsystem=[%s] bWasSuccessful=[%s] Error=[%s]."), GetLocalUserNum(), *ListName, ToString(SubsystemType), *LexToString(bWasSuccessful), *ErrorStr);
 	if (bWasSuccessful)
 	{
 		TArray<TSharedRef<FOnlineFriend>> FriendsList;
@@ -830,7 +830,7 @@ void USocialToolkit::HandleQueryBlockedPlayersComplete(const FUniqueNetId& UserI
 {
 	if (UserId == *GetLocalUserNetId(SubsystemType))
 	{
-		UE_LOG(LogParty, Log, TEXT("SocialToolkit [%d] finished querying blocked players on subsystem [%s] with error [%s]."), GetLocalUserNum(), ToString(SubsystemType), *ErrorStr);
+		UE_LOG(LogParty, Log, TEXT("SocialToolkit [%d] finished querying blocked players Subsystem=[%s] bWasSuccessful=[%s] Error=[%s]."), GetLocalUserNum(), ToString(SubsystemType), *LexToString(bWasSuccessful), *ErrorStr);
 
 		if (bWasSuccessful)
 		{
@@ -855,7 +855,7 @@ void USocialToolkit::HandleQueryRecentPlayersComplete(const FUniqueNetId& UserId
 {
 	if (UserId == *GetLocalUserNetId(SubsystemType))
 	{
-		UE_LOG(LogParty, Log, TEXT("SocialToolkit [%d] finished querying recent player list [%s] on subsystem [%s] with error [%s]."), GetLocalUserNum(), *Namespace, ToString(SubsystemType), *ErrorStr);
+		UE_LOG(LogParty, Log, TEXT("SocialToolkit [%d] finished querying recent player list Namespace=[%s] Subsystem=[%s] bWasSuccessful=[%s] Error=[%s]."), GetLocalUserNum(), *Namespace, ToString(SubsystemType), *LexToString(bWasSuccessful), *ErrorStr);
 
 		if (bWasSuccessful)
 		{
