Modify update manager to require a context name when starting the update process.
Update contexts allow defining update check behavior based on the call-site and are config driven / hotfixable.

diff --git a/Engine/Plugins/Online/OnlineFramework/Source/Hotfix/Private/UpdateManager.cpp b/Engine/Plugins/Online/OnlineFramework/Source/Hotfix/Private/UpdateManager.cpp
--- a/Engine/Plugins/Online/OnlineFramework/Source/Hotfix/Private/UpdateManager.cpp
+++ b/Engine/Plugins/Online/OnlineFramework/Source/Hotfix/Private/UpdateManager.cpp
@@ -1,6 +1,7 @@
 // Copyright Epic Games, Inc. All Rights Reserved.
 
 #include "UpdateManager.h"
+#include "Algo/Transform.h"
 #include "Misc/CommandLine.h"
 #include "Containers/Ticker.h"
 #include "HAL/IConsoleManager.h"
@@ -100,11 +101,26 @@ void UUpdateManager::StartCheck(bool bInCheckHotfixOnly)
 	StartCheckInternal(bInCheckHotfixOnly);
 }
 
+void UUpdateManager::StartUpdateCheck(const FString& ContextName)
+{
+	StartUpdateCheckInternal(GetContextDefinition(ContextName));
+}
+
 UUpdateManager::EUpdateStartResult UUpdateManager::StartCheckInternal(bool bInCheckHotfixOnly)
+{
+	FUpdateContextDefinition AdhocDefinition;
+	AdhocDefinition.Name = TEXT("Adhoc");
+	AdhocDefinition.bCheckAvailabilityOnly = bInCheckHotfixOnly;
+	return StartUpdateCheckInternal(AdhocDefinition);
+}
+
+UUpdateManager::EUpdateStartResult UUpdateManager::StartUpdateCheckInternal(const FUpdateContextDefinition& ContextDefinition)
 {
 	EUpdateStartResult Result = EUpdateStartResult::None;
 
-	if (!ChecksEnabled())
+	UE_LOG(LogHotfixManager, Log, TEXT("Starting update check using context \"%s\""), *ContextDefinition.Name);
+
+	if (!ChecksEnabled() || !ContextDefinition.bEnabled)
 	{
 		UE_LOG(LogHotfixManager, Display, TEXT("Update checks disabled!"));
 		bInitialUpdateFinished = true;
@@ -126,7 +142,9 @@ UUpdateManager::EUpdateStartResult UUpdateManager::StartCheckInternal(bool bInCh
 		CurrentUpdateState == EUpdateState::UpdatePending ||
 		CurrentUpdateState == EUpdateState::UpdateComplete))
 	{
-		bCheckHotfixAvailabilityOnly = bInCheckHotfixOnly;
+		bCheckHotfixAvailabilityOnly = ContextDefinition.bCheckAvailabilityOnly;
+		bCurrentUpdatePatchCheckEnabled = ContextDefinition.bPatchCheckEnabled;
+		bCurrentUpdatePlatformEnvironmentDetectionEnabled = ContextDefinition.bPlatformEnvironmentDetectionEnabled;
 
 		// Immediately move into a pending state so the UI state trigger fires
 		SetUpdateState(EUpdateState::UpdatePending);
@@ -226,6 +244,8 @@ void UUpdateManager::CheckComplete(EUpdateCompletionStatus Result, bool bUpdateT
 		}
 
 		bCheckHotfixAvailabilityOnly = false;
+		bCurrentUpdatePatchCheckEnabled = true;
+		bCurrentUpdatePlatformEnvironmentDetectionEnabled = true;
 		
 		OnUpdateCheckComplete().Broadcast(FinalResult);
 	};
@@ -247,7 +267,7 @@ void UUpdateManager::StartPatchCheck()
 	check(GameInstance);
 
 	SetUpdateState(EUpdateState::CheckingForPatch);
-	if (GameInstance->IsDedicatedServerInstance())
+	if (GameInstance->IsDedicatedServerInstance() || !bCurrentUpdatePatchCheckEnabled)
 	{
 		PatchCheckComplete(EPatchCheckResult::NoPatchRequired);
 		return;
@@ -353,7 +373,7 @@ void UUpdateManager::InstallBundlePatchCheckComplete(EInstallBundleManagerPatchC
 
 void UUpdateManager::StartPlatformEnvironmentCheck()
 {
-	if (!bPlatformEnvironmentDetected)
+	if (bCurrentUpdatePlatformEnvironmentDetectionEnabled && !bPlatformEnvironmentDetected)
 	{
 #if UPDATEMANAGER_PLATFORM_ENVIRONMENT_DETECTION
 		if (DetectPlatformEnvironment())
@@ -596,6 +616,26 @@ bool UUpdateManager::Tick(float InDeltaTime)
 	return true;
 }
 
+void UUpdateManager::PostInitProperties()
+{
+	Super::PostInitProperties();
+
+	if (!HasAnyFlags(RF_ClassDefaultObject))
+	{
+		PopulateContextDefinitions();
+	}
+}
+
+void UUpdateManager::PostReloadConfig(FProperty* PropertyThatWasLoaded)
+{
+	Super::PostReloadConfig(PropertyThatWasLoaded);
+
+	if (!HasAnyFlags(RF_ClassDefaultObject))
+	{
+		PopulateContextDefinitions();
+	}
+}
+
 float UUpdateManager::GetLoadProgress() const
 {
 	const int32 CurrentNumFiles = GetNumAsyncPackages();
@@ -653,6 +693,30 @@ FTSTicker::FDelegateHandle UUpdateManager::DelayResponse(DelayCb&& Delegate, flo
 		}), Delay);
 }
 
+void UUpdateManager::PopulateContextDefinitions()
+{
+	ProcessedUpdateContextDefinitions.Empty();
+	ProcessedUpdateContextDefinitions.Reserve(UpdateContextDefinitions.Num());
+
+	// Elements later in the list will override elements with the same name which appear earlier.
+	Algo::Transform(UpdateContextDefinitions, ProcessedUpdateContextDefinitions,
+	[](const FUpdateContextDefinition& Definition)
+	{
+		return TTuple<FString, FUpdateContextDefinition>(Definition.Name, Definition);
+	});
+}
+
+const FUpdateContextDefinition& UUpdateManager::GetContextDefinition(const FString& ContextName) const
+{
+	if (const FUpdateContextDefinition* Definiton = ProcessedUpdateContextDefinitions.Find(ContextName))
+	{
+		return *Definiton;
+	}
+
+	UE_LOG(LogHotfixManager, Warning, TEXT("Failed to find update context \"%s\". Using \"Unknown\" definition."), *ContextName);
+	return UpdateContextDefinitionUnknown;
+}
+
 UWorld* UUpdateManager::GetWorld() const
 {
 	UGameInstance* GameInstance = GetTypedOuter<UGameInstance>();
diff --git a/Engine/Plugins/Online/OnlineFramework/Source/Hotfix/Public/UpdateManager.h b/Engine/Plugins/Online/OnlineFramework/Source/Hotfix/Public/UpdateManager.h
--- a/Engine/Plugins/Online/OnlineFramework/Source/Hotfix/Public/UpdateManager.h
+++ b/Engine/Plugins/Online/OnlineFramework/Source/Hotfix/Public/UpdateManager.h
@@ -105,6 +105,36 @@ typedef FOnUpdateHotfixProgress::FDelegate FOnUpdateHotfixProgressDelegate;
 DECLARE_MULTICAST_DELEGATE_TwoParams(FOnUpdateHotfixProcessedFile, const FString& /** FriendlyName */, const FString& /** CachedFileName */);
 typedef FOnUpdateHotfixProcessedFile::FDelegate FOnUpdateHotfixProcessedFileDelegate;
 
+/**
+ * Update context definition loaded from config.
+ *
+ * Defines the behavior to use during an update check.
+ */ 
+USTRUCT()
+struct HOTFIX_API FUpdateContextDefinition
+{
+	GENERATED_BODY()
+public:
+	/** The name used to lookup the context definition. */
+	UPROPERTY(Config)
+	FString Name;
+	/** Whether to perform any checks. */
+	UPROPERTY(Config)
+	bool bEnabled = true;
+	/** Whether to perform an update, or to only check for the presence of an update. */
+	UPROPERTY(Config)
+	bool bCheckAvailabilityOnly = false;
+	/** Whether to check for patches during the update process. */
+	UPROPERTY(Config)
+	bool bPatchCheckEnabled = true;
+	/** Whether to attempt detection of the platforms environment. */
+	UPROPERTY(Config)
+	bool bPlatformEnvironmentDetectionEnabled = true;
+	/** Additional tags for extensibility. */
+	UPROPERTY(Config)
+	TSet<FString> AdditionalTags;
+};
+
 /**
  * Update manager
  *
@@ -145,7 +175,15 @@ public:
 	 *
 	 * @param bInCheckHotfixOnly check for the availability of a hotfix only (does not apply)
 	 */
+	UE_DEPRECATED_FORGAME(5.1, "UUpdateManager::StartCheck is deprecated, please use StartUpdateCheck instead.")
 	virtual void StartCheck(bool bInCheckHotfixOnly = false);
+
+	/**
+	 * Start an update check using properties defined by context.
+	 *
+	 * @param ContextName The context definitions to use when updating.
+	 */
+	void StartUpdateCheck(const FString& ContextName);
 	
 	/**
 	 *  @return the load progress (0..1)
@@ -235,11 +273,22 @@ protected:
 	 *
 	 * @return whether or not the check started, returned a cached value, or did nothing (already in progress)
 	 */
+	UE_DEPRECATED_FORGAME(5.1, "UUpdateManager::StartCheckInternal is deprecated, please override StartCheckUpdateInternal instead.")
 	virtual EUpdateStartResult StartCheckInternal(bool bInCheckHotfixOnly);
 
+	/** 
+	 * Internal call for StartUpdateCheck
+	 *
+	 * @return whether or not the check started, returned a cached value, or did nothing (already in progress)
+	 */
+	virtual EUpdateStartResult StartUpdateCheckInternal(const FUpdateContextDefinition& ContextDefinition);
+
 	/** Tick function during initial preload */
 	virtual bool Tick(float DeltaTime);
 
+	virtual void PostInitProperties() override;
+	virtual void PostReloadConfig(FProperty* PropertyThatWasLoaded) override;
+
 	/** Amount of time to wait between the internal hotfix check completing and advancing to the next stage */
 	UPROPERTY(Config)
 	float HotfixCheckCompleteDelay;
@@ -409,6 +458,23 @@ private:
 
 	friend bool SkipPatchCheck(UUpdateManager* UpdateManager);
 
+private:
+	void PopulateContextDefinitions();
+	const FUpdateContextDefinition& GetContextDefinition(const FString& ContextName) const;
+
+	// Definition to use when a context lookup fails.
+	UPROPERTY(Config)
+	FUpdateContextDefinition UpdateContextDefinitionUnknown;
+
+	// Definitions for update behavior.
+	UPROPERTY(Config)
+	TArray<FUpdateContextDefinition> UpdateContextDefinitions;
+
+	// Update context definitions which have been processed following a config update.
+	TMap<FString, FUpdateContextDefinition> ProcessedUpdateContextDefinitions;
+
+	bool bCurrentUpdatePatchCheckEnabled = true;
+	bool bCurrentUpdatePlatformEnvironmentDetectionEnabled = true;
 protected:
 
 	/** @return a pointer to the hotfix manager */
