Fix case where a party member display name may show incorrect name when a user restores a party on a different platform from the one last played on.

diff --git a/Engine/Plugins/Online/OnlineFramework/Source/Party/Private/Party/PartyMember.cpp b/Engine/Plugins/Online/OnlineFramework/Source/Party/Private/Party/PartyMember.cpp
--- a/Engine/Plugins/Online/OnlineFramework/Source/Party/Private/Party/PartyMember.cpp
+++ b/Engine/Plugins/Online/OnlineFramework/Source/Party/Private/Party/PartyMember.cpp
@@ -70,6 +70,7 @@ void UPartyMember::InitializePartyMember(const FOnlinePartyMemberConstRef& InOss
 	{
 		OssPartyMember = InOssMember;
 		OssPartyMember->OnMemberConnectionStatusChanged().AddUObject(this, &ThisClass::HandleMemberConnectionStatusChanged);
+		OssPartyMember->OnMemberAttributeChanged().AddUObject(this, &ThisClass::HandleMemberAttributeChanged);
 		USocialToolkit* OwnerToolkit = GetParty().GetSocialManager().GetSocialToolkit(OssPartyMember->GetUserId());
 		// If we are not a local user then we simply get the first local user's toolkit
 		if (OwnerToolkit == nullptr)
@@ -278,3 +279,11 @@ void UPartyMember::HandleMemberConnectionStatusChanged(const FUniqueNetId& Chang
 {
 	OnMemberConnectionStatusChanged().Broadcast();
 }
+
+void UPartyMember::HandleMemberAttributeChanged(const FUniqueNetId& ChangedUserId, const FString& Attribute, const FString& NewValue, const FString& OldValue)
+{
+	if (Attribute == USER_ATTR_DISPLAYNAME)
+	{
+		OnDisplayNameChanged().Broadcast();
+	}
+}
diff --git a/Engine/Plugins/Online/OnlineFramework/Source/Party/Public/Party/PartyMember.h b/Engine/Plugins/Online/OnlineFramework/Source/Party/Public/Party/PartyMember.h
--- a/Engine/Plugins/Online/OnlineFramework/Source/Party/Public/Party/PartyMember.h
+++ b/Engine/Plugins/Online/OnlineFramework/Source/Party/Public/Party/PartyMember.h
@@ -109,6 +109,7 @@ public:
 	FOnPartyMemberStateChanged& OnPromotedToLeader() const { return OnPromotedToLeaderEvent; }
 	FOnPartyMemberStateChanged& OnDemoted() const { return OnDemotedEvent; }
 	FOnPartyMemberStateChanged& OnMemberConnectionStatusChanged() const { return OnMemberConnectionStatusChangedEvent; }
+	FOnPartyMemberStateChanged& OnDisplayNameChanged() const { return OnDisplayNameChangedEvent; }
 
 	DECLARE_EVENT_OneParam(UPartyMember, FOnPartyMemberLeft, EMemberExitedReason)
 	FOnPartyMemberLeft& OnLeftParty() const { return OnLeftPartyEvent; }
@@ -141,6 +142,7 @@ protected:
 private:
 	void HandleSocialUserInitialized(USocialUser& InitializedUser);
 	void HandleMemberConnectionStatusChanged(const FUniqueNetId& ChangedUserId, const EMemberConnectionStatus NewMemberConnectionStatus, const EMemberConnectionStatus PreviousMemberConnectionStatus);
+	void HandleMemberAttributeChanged(const FUniqueNetId& ChangedUserId, const FString& Attribute, const FString& NewValue, const FString& OldValue);
 
 	FOnlinePartyMemberConstPtr OssPartyMember;
 
@@ -149,6 +151,7 @@ private:
 
 	bool bHasReceivedInitialData = false;
 	mutable FOnPartyMemberStateChanged OnMemberConnectionStatusChangedEvent;
+	mutable FOnPartyMemberStateChanged OnDisplayNameChangedEvent;
 	mutable FOnPartyMemberStateChanged OnMemberInitializedEvent;
 	mutable FOnPartyMemberStateChanged OnPromotedToLeaderEvent;
 	mutable FOnPartyMemberStateChanged OnDemotedEvent;
