Make social user list keep weak pointers to users to fix GC errors when leaving PIE with pending-kill disabled.

diff --git a/Engine/Plugins/Online/OnlineFramework/Source/Party/Private/User/SocialUserList.cpp b/Engine/Plugins/Online/OnlineFramework/Source/Party/Private/User/SocialUserList.cpp
--- a/Engine/Plugins/Online/OnlineFramework/Source/Party/Private/User/SocialUserList.cpp
+++ b/Engine/Plugins/Online/OnlineFramework/Source/Party/Private/User/SocialUserList.cpp
@@ -89,11 +89,6 @@ void FSocialUserList::InitializeList()
 	SetAllowAutoUpdate(ListConfig.bAutoUpdate);
 }
 
-void FSocialUserList::AddReferencedObjects(FReferenceCollector& Collector)
-{
-	Collector.AddReferencedObjects(Users);
-}
-
 void FSocialUserList::UpdateNow()
 {
 	UpdateListInternal();
@@ -129,9 +124,9 @@ void FSocialUserList::HandleOwnerToolkitReset()
 {
 	const bool bTriggerChangeEvent = Users.Num() > 0;
 
-	for (const USocialUser* User : Users)
+	for (TWeakObjectPtr<USocialUser> UserPtr : Users)
 	{
-		if (ensureMsgf(User, TEXT("Encountered a nullptr entry in FSocialUserList::Users array!")))
+		if (USocialUser* User = UserPtr.Get(); ensureMsgf(User, TEXT("Encountered a nullptr entry in FSocialUserList::Users array!")))
 		{
 			OnUserRemoved().Broadcast(*User);
 		}
@@ -486,7 +481,7 @@ struct FUserSortData
 		, CustomSortValueTertiary(InCustomSortValueTertiary) 
 	{ }
 
-	USocialUser* User;
+	USocialUser* User = nullptr;
 	EOnlinePresenceState::Type OnlineStatus;
 	bool PlayingThisGame;
 	FString DisplayName;
@@ -576,7 +571,7 @@ void FSocialUserList::UpdateListInternal()
 		bListUpdated = true;
 
 		Users.RemoveAllSwap(
-			[this] (USocialUser* User)
+			[this] (TWeakObjectPtr<USocialUser> User)
 		{
 				if (PendingRemovals.Contains(User))
 				{
@@ -596,9 +591,12 @@ void FSocialUserList::UpdateListInternal()
 		bListUpdated = true;
 		Users.Append(PendingAdds);
 
-		for (USocialUser* User : PendingAdds)
+		for (TWeakObjectPtr<USocialUser> UserPtr : PendingAdds)
 		{
-			OnUserAdded().Broadcast(*User);
+			if (USocialUser* User = UserPtr.Get(); ensure(User))
+			{
+				OnUserAdded().Broadcast(*User);
+			}
 		}
 		PendingAdds.Reset();
 	}
@@ -620,18 +618,19 @@ void FSocialUserList::UpdateListInternal()
 				TArray<FUserSortData> SortedData;
 				SortedData.Reserve(NumUsers);
 
-				Algo::Transform(Users, SortedData, [](USocialUser* const User) -> FUserSortData
+				Algo::TransformIf(Users, SortedData, 
+					[](const TWeakObjectPtr<USocialUser> UserPtr) { return UserPtr.IsValid(); }, 
+					[](const TWeakObjectPtr<USocialUser> UserPtr) -> FUserSortData
 				{
+					USocialUser* User = UserPtr.Get();
 					return FUserSortData(User, User->GetOnlineStatus(), User->IsPlayingThisGame(), User->GetDisplayName(), User->GetCustomSortValuePrimary(), User->GetCustomSortValueSecondary(), User->GetCustomSortValueTertiary());
 				});
 
 				Algo::Sort(SortedData);
 
 				// replace contents of Users from SortedData array
-				for (int Index = 0; Index < NumUsers; Index++)
-				{
-					Users[Index] = SortedData[Index].User;
-				}
+				Users.Reset(SortedData.Num());
+				Algo::Transform(SortedData, Users, [](const FUserSortData& Data){ return Data.User; });
 			}
 		}
 		else
diff --git a/Engine/Plugins/Online/OnlineFramework/Source/Party/Private/User/SocialUserList.h b/Engine/Plugins/Online/OnlineFramework/Source/Party/Private/User/SocialUserList.h
--- a/Engine/Plugins/Online/OnlineFramework/Source/Party/Private/User/SocialUserList.h
+++ b/Engine/Plugins/Online/OnlineFramework/Source/Party/Private/User/SocialUserList.h
@@ -6,31 +6,27 @@
 #include "UObject/GCObject.h"
 #include "Containers/Ticker.h"
 
+class USocialUser;
+
 enum class EMemberExitedReason : uint8;
 typedef TSharedRef<const IOnlinePartyRequestToJoinInfo> IOnlinePartyRequestToJoinInfoConstRef;
 
-class FSocialUserList : public ISocialUserList, public FGCObject, public TSharedFromThis<FSocialUserList>
+class FSocialUserList : public ISocialUserList, public TSharedFromThis<FSocialUserList>
 {
 	friend class USocialToolkit;
 public:
 	static TSharedRef<FSocialUserList> CreateUserList(const USocialToolkit& InOwnerToolkit, const FSocialUserListConfig& Config);
 
-	virtual void AddReferencedObjects(FReferenceCollector& Collector) override;
-	virtual FString GetReferencerName() const override
-	{
-		return TEXT("FSocialUserList");
-	}
-
 	FOnUserAdded& OnUserAdded() const override { return OnUserAddedEvent; }
 	FOnUserRemoved& OnUserRemoved() const override { return OnUserRemovedEvent; }
 	FOnUpdateComplete& OnUpdateComplete() const override { return OnUpdateCompleteEvent; }
 
 	virtual FString GetListName() const override { return ListConfig.Name; }
 
-	void UpdateNow();
-	void SetAllowAutoUpdate(bool bIsEnabled);
-	void SetAllowSortDuringUpdate(bool bIsEnabled);
-	const TArray<USocialUser*>& GetUsers() const { return Users; }
+	virtual void UpdateNow() override;
+	virtual void SetAllowAutoUpdate(bool bIsEnabled) override;
+	virtual void SetAllowSortDuringUpdate(bool bIsEnabled) override;
+	virtual const TArray<TWeakObjectPtr<USocialUser>>& GetUsers() const override { return Users; }
 
 	bool HasPresenceFilters() const;
 
@@ -89,12 +85,8 @@ private:
 	void InitializeList();
 
 	TWeakObjectPtr<const USocialToolkit> OwnerToolkit;
-
-	UPROPERTY()
-	TArray<USocialUser*> Users;
-
-	UPROPERTY()
-	TArray<USocialUser*> PendingAdds;
+	TArray<TWeakObjectPtr<USocialUser>> Users;
+	TArray<TWeakObjectPtr<USocialUser>> PendingAdds;
 
 	TSet<TWeakObjectPtr<USocialUser>> UsersWithDirtyPresence;
 	TArray<TWeakObjectPtr<const USocialUser>> PendingRemovals;
diff --git a/Engine/Plugins/Online/OnlineFramework/Source/Party/Public/User/ISocialUserList.h b/Engine/Plugins/Online/OnlineFramework/Source/Party/Public/User/ISocialUserList.h
--- a/Engine/Plugins/Online/OnlineFramework/Source/Party/Public/User/ISocialUserList.h
+++ b/Engine/Plugins/Online/OnlineFramework/Source/Party/Public/User/ISocialUserList.h
@@ -91,7 +91,7 @@ public:
 	DECLARE_EVENT(ISocialUserList, FOnUpdateComplete)
 	virtual FOnUpdateComplete& OnUpdateComplete() const = 0;
 
-	virtual const TArray<USocialUser*>& GetUsers() const = 0;
+	virtual const TArray<TWeakObjectPtr<USocialUser>>& GetUsers() const = 0;
 	virtual FString GetListName() const = 0;
 
 	/** Trigger an update of the list immediately, regardless of auto update period */
