Add new custom invite type in ESocialPartyInviteMethod

diff --git a/Engine/Plugins/Online/OnlineFramework/Source/Party/Public/Party/PartyTypes.h b/Engine/Plugins/Online/OnlineFramework/Source/Party/Public/Party/PartyTypes.h
--- a/Engine/Plugins/Online/OnlineFramework/Source/Party/Public/Party/PartyTypes.h
+++ b/Engine/Plugins/Online/OnlineFramework/Source/Party/Public/Party/PartyTypes.h
@@ -230,7 +230,9 @@ enum class ESocialPartyInviteMethod : uint8
 	/** Default value for try invite */
 	Other = 0,
 	/** Invite was sent from a toast */
-	Notification
+	Notification,
+	/** Invite was sent with a custom method */
+	Custom1
 };
 
 UENUM()
