Add CanAutoRecreatePersistentParty method to control whether we should automatically recreate a persistent party on unexpected loss.
The previous behavior only checked for whether the user was logged in.  This allows for additional checks for this specific behavior.

diff --git a/Engine/Plugins/Online/OnlineFramework/Source/Party/Private/SocialManager.cpp b/Engine/Plugins/Online/OnlineFramework/Source/Party/Private/SocialManager.cpp
--- a/Engine/Plugins/Online/OnlineFramework/Source/Party/Private/SocialManager.cpp
+++ b/Engine/Plugins/Online/OnlineFramework/Source/Party/Private/SocialManager.cpp
@@ -859,6 +859,8 @@ void USocialManager::FinishJoinPartyAttempt(FJoinPartyAttempt& JoinAttemptToDest
 	if (bWasPersistentPartyJoinAttempt && !JoinResult.WasSuccessful() && !GetPersistentParty())
 	{
 		// Something goofed when trying to join a new persistent party, so create a replacement immediately
+		// TODO:  See if the below warning happens, and if so, flesh out this flow
+		UE_CLOG(!GetFirstLocalUserToolkit()->CanAutoRecreatePersistentParty(), LogParty, Warning, TEXT("FinishJoinPartyAttempt creating a persistent party when we shouldn't!"));
 		CreatePersistentParty();
 	}
 }
@@ -1324,7 +1326,7 @@ void USocialManager::HandlePartyLeft(EMemberExitedReason Reason, USocialParty* L
 			}
 		}
 
-		if (LeftParty->IsPersistentParty() && GetFirstLocalUserToolkit()->IsOwnerLoggedIn())
+		if (LeftParty->IsPersistentParty() && GetFirstLocalUserToolkit()->CanAutoRecreatePersistentParty())
 		{
 			UE_LOG(LogParty, Verbose, TEXT("Finished leaving persistent party without a join/rejoin target. Creating a new persistent party now."));
 
@@ -1338,7 +1340,7 @@ void USocialManager::HandleLeavePartyForMissingJoinAttempt(const FUniqueNetId& L
 {
 	ABORT_DURING_SHUTDOWN();
 
-	if (PartyTypeId == IOnlinePartySystem::GetPrimaryPartyTypeId() && GetFirstLocalUserToolkit()->IsOwnerLoggedIn() && !GetPersistentPartyInternal(true))
+	if (PartyTypeId == IOnlinePartySystem::GetPrimaryPartyTypeId() && GetFirstLocalUserToolkit()->CanAutoRecreatePersistentParty() && !GetPersistentPartyInternal(true))
 	{
 		// We just had to bail on the persistent party due to unforeseen shenanigans, so try to correct things and set another one back up
 		CreatePersistentParty();
diff --git a/Engine/Plugins/Online/OnlineFramework/Source/Party/Public/SocialToolkit.h b/Engine/Plugins/Online/OnlineFramework/Source/Party/Public/SocialToolkit.h
--- a/Engine/Plugins/Online/OnlineFramework/Source/Party/Public/SocialToolkit.h
+++ b/Engine/Plugins/Online/OnlineFramework/Source/Party/Public/SocialToolkit.h
@@ -142,7 +142,13 @@ public:
 	virtual void NotifyPartyInviteReceived(USocialUser& SocialUser, const IOnlinePartyJoinInfo& Invite);
 	virtual void NotifyPartyInviteRemoved(USocialUser& SocialUser, const IOnlinePartyJoinInfo& Invite);
 
-	void RequestToJoinParty(USocialUser& SocialUser);
+	/**
+	 * Can we automatically re-create a persistent party, for example on party join failure or on party kick?
+	 * Typically you will want to, so the user always has a party, but there may be scenarios where you want to delay
+	 * creating a new party until the client is in a better state.
+	 * @return true if we can automatically recreate a persistent party
+	 */
+	virtual bool CanAutoRecreatePersistentParty() const { return IsOwnerLoggedIn(); }
 
 #if WITH_EDITOR
 	bool Debug_IsRandomlyChangingPresence() const { return bDebug_IsRandomlyChangingUserPresence; }
