Make Party Invite Reject a virtual method in SocialUser

#tests verified that it build correctly under win64.

diff --git a/Engine/Plugins/Online/OnlineFramework/Source/Party/Public/User/SocialUser.h b/Engine/Plugins/Online/OnlineFramework/Source/Party/Public/User/SocialUser.h
--- a/Engine/Plugins/Online/OnlineFramework/Source/Party/Public/User/SocialUser.h
+++ b/Engine/Plugins/Online/OnlineFramework/Source/Party/Public/User/SocialUser.h
@@ -129,7 +129,7 @@ public:
 	void JoinParty(const FOnlinePartyTypeId& PartyTypeId) const;
 
 	virtual void JoinParty(const FOnlinePartyTypeId& PartyTypeId, const FName& JoinMethod) const;
-	void RejectPartyInvite(const FOnlinePartyTypeId& PartyTypeId);
+	virtual void RejectPartyInvite(const FOnlinePartyTypeId& PartyTypeId);
 
 	bool HasBeenInvitedToParty(const FOnlinePartyTypeId& PartyTypeId) const;
 	bool CanInviteToParty(const FOnlinePartyTypeId& PartyTypeId) const;
