Fixed join method not being recovered after disconnection.

diff --git a/Engine/Plugins/Online/OnlineFramework/Source/Party/Private/Party/PartyMember.cpp b/Engine/Plugins/Online/OnlineFramework/Source/Party/Private/Party/PartyMember.cpp
--- a/Engine/Plugins/Online/OnlineFramework/Source/Party/Private/Party/PartyMember.cpp
+++ b/Engine/Plugins/Online/OnlineFramework/Source/Party/Private/Party/PartyMember.cpp
@@ -6,7 +6,7 @@
 #include "SocialManager.h"
 #include "SocialToolkit.h"
 
-#include "OnlineSubsystem.h"
+#include "OnlineSubsystemUtils.h"
 #include "Interfaces/OnlinePartyInterface.h"
 
 //////////////////////////////////////////////////////////////////////////
@@ -108,13 +108,29 @@ void UPartyMember::InitializeLocalMemberRepData()
 	
 	const USocialParty& CurrentParty = GetParty();
 	
+	FString JoinMethod;
 	if (const USocialManager::FJoinPartyAttempt* JoinAttempt = CurrentParty.GetSocialManager().GetJoinAttemptInProgress(CurrentParty.GetPartyTypeId()))
 	{
-		const FString JoinMethod = JoinAttempt->JoinMethod.ToString();
-
-		UE_LOG(LogParty, Verbose, TEXT("Join method for local member [%s] is %s."), *ToDebugString(), *JoinMethod);
-		MemberDataReplicator->SetJoinMethod(JoinMethod);
+		JoinMethod = JoinAttempt->JoinMethod.ToString();
+		UE_LOG(LogParty, Verbose, TEXT("Join method from join attempt for local member is %s."), *JoinMethod);
 	}
+	else
+	{
+		IOnlinePartyPtr PartyInterface = Online::GetPartyInterfaceChecked(GetWorld());
+		FOnlinePartyDataConstPtr PartyMemberData = PartyInterface.IsValid() ? PartyInterface->GetPartyMemberData(*CurrentParty.GetOwningLocalUserId(), CurrentParty.GetPartyId(), *GetPrimaryNetId(), DefaultPartyDataNamespace) : nullptr;
+		
+		if (PartyMemberData.IsValid())
+		{
+			FVariantData AttrValue;
+			if (PartyMemberData->GetAttribute(TEXT("JoinMethod"), AttrValue))
+			{
+				JoinMethod = AttrValue.ToString();
+				UE_LOG(LogParty, Verbose, TEXT("Join method recovered from the Party member data is %s."), *JoinMethod);
+			}
+		}
+	}
+
+	MemberDataReplicator->SetJoinMethod(JoinMethod);
 }
 
 void UPartyMember::Shutdown()
diff --git a/Engine/Plugins/Online/OnlineFramework/Source/Party/Public/SocialManager.h b/Engine/Plugins/Online/OnlineFramework/Source/Party/Public/SocialManager.h
--- a/Engine/Plugins/Online/OnlineFramework/Source/Party/Public/SocialManager.h
+++ b/Engine/Plugins/Online/OnlineFramework/Source/Party/Public/SocialManager.h
@@ -143,7 +143,7 @@ protected:
 
 		TWeakObjectPtr<const USocialUser> TargetUser;
 		FOnlinePartyTypeId PartyTypeId;
-		FName JoinMethod;
+		FName JoinMethod = PartyJoinMethod::Unspecified;
 		FUniqueNetIdRepl TargetUserPlatformId;
 		FSessionId PlatformSessionId;
 
