[KillPendingKill] SocialParty reservation Beacon is now a WeakPtr

diff --git a/Engine/Plugins/Online/OnlineFramework/Source/Party/Private/Party/SocialParty.cpp b/Engine/Plugins/Online/OnlineFramework/Source/Party/Private/Party/SocialParty.cpp
--- a/Engine/Plugins/Online/OnlineFramework/Source/Party/Private/Party/SocialParty.cpp
+++ b/Engine/Plugins/Online/OnlineFramework/Source/Party/Private/Party/SocialParty.cpp
@@ -756,7 +756,7 @@ void USocialParty::HandlePartyJoinRequestReceived(const FUniqueNetId& LocalUserI
 		PendingApproval.bIsJIPApproval = false;
 		PendingApprovals.Enqueue(PendingApproval);
 
-		if (!ReservationBeaconClient && JoinApproval.GetApprovalAction() == EApprovalAction::EnqueueAndStartBeacon)
+		if (!ReservationBeaconClient.Get() && JoinApproval.GetApprovalAction() == EApprovalAction::EnqueueAndStartBeacon)
 		{
 			ConnectToReservationBeacon();
 		}
@@ -796,7 +796,7 @@ void USocialParty::RemovePlayerFromReservationBeacon(const FUniqueNetId& LocalUs
 	PendingApprovals.Enqueue(PendingApproval);
 
 
-	if (!ReservationBeaconClient)
+	if (!ReservationBeaconClient.Get())
 	{
 		ConnectToReservationBeacon();
 	}
@@ -835,7 +835,7 @@ void USocialParty::HandlePartyJIPRequestReceived(const FUniqueNetId& LocalUserId
 		PendingApproval.bIsJIPApproval = true;
 		PendingApprovals.Enqueue(MoveTemp(PendingApproval));
 
-		if (!ReservationBeaconClient && JoinApproval.GetApprovalAction() == EApprovalAction::EnqueueAndStartBeacon)
+		if (!ReservationBeaconClient.Get() && JoinApproval.GetApprovalAction() == EApprovalAction::EnqueueAndStartBeacon)
 		{
 			ConnectToReservationBeacon();
 		}
@@ -1442,7 +1442,8 @@ bool USocialParty::IsPartyLeaderLocal() const
 bool USocialParty::IsNetDriverFromReservationBeacon(const UNetDriver* const InNetDriver) const
 {
 	const FName NetDriverName = InNetDriver->NetDriverName;
-	return (ReservationBeaconClient && NetDriverName == ReservationBeaconClient->GetNetDriverName()) || (NetDriverName == LastReservationBeaconClientNetDriverName);
+	APartyBeaconClient* LocalReservationBeaconClient = ReservationBeaconClient.Get();
+	return (LocalReservationBeaconClient && NetDriverName == LocalReservationBeaconClient->GetNetDriverName()) || (NetDriverName == LastReservationBeaconClientNetDriverName);
 }
 
 FString USocialParty::ToDebugString() const
@@ -1488,7 +1489,8 @@ void USocialParty::HandlePartyStateChanged(const FUniqueNetId& LocalUserId, cons
 
 void USocialParty::ConnectToReservationBeacon()
 {
-	if (IsLocalPlayerPartyLeader() && !ReservationBeaconClient)
+	APartyBeaconClient* LocalReservationBeaconClient = ReservationBeaconClient.Get();
+	if (IsLocalPlayerPartyLeader() && !LocalReservationBeaconClient)
 	{
 		FPendingMemberApproval NextApproval;
 		if (PendingApprovals.Peek(NextApproval))
@@ -1511,12 +1513,12 @@ void USocialParty::ConnectToReservationBeacon()
 					{
 						// Reconnect to the reservation beacon to maintain our place in the game (just until actual joined, holds place for all party members)
 						ReservationBeaconClient = World->SpawnActor<APartyBeaconClient>(ReservationBeaconClientClass);
-						if (ReservationBeaconClient)
+						if (LocalReservationBeaconClient)
 						{
-							UE_LOG(LogParty, Verbose, TEXT("Party [%s] created reservation beacon [%s]."), *ToDebugString(), *ReservationBeaconClient->GetName());
+							UE_LOG(LogParty, Verbose, TEXT("Party [%s] created reservation beacon [%s]."), *ToDebugString(), *LocalReservationBeaconClient->GetName());
 
-							ReservationBeaconClient->OnHostConnectionFailure().BindUObject(this, &USocialParty::HandleBeaconHostConnectionFailed);
-							ReservationBeaconClient->OnReservationRequestComplete().BindUObject(this, &USocialParty::HandleReservationRequestComplete);
+							LocalReservationBeaconClient->OnHostConnectionFailure().BindUObject(this, &USocialParty::HandleBeaconHostConnectionFailed);
+							LocalReservationBeaconClient->OnReservationRequestComplete().BindUObject(this, &USocialParty::HandleReservationRequestComplete);
 
 							TArray<FPlayerReservation> ReservationAsArray;
 							ReservationAsArray.Reserve(NextApproval.Members.Num());
@@ -1575,7 +1577,8 @@ void USocialParty::RejectAllPendingJoinRequests()
 
 void USocialParty::HandleBeaconHostConnectionFailed()
 {
-	UE_LOG(LogParty, Verbose, TEXT("Host connection failed for reservation beacon [%s]"), ReservationBeaconClient ? *ReservationBeaconClient->GetName() : TEXT(""));
+	APartyBeaconClient* LocalReservationBeaconClient = ReservationBeaconClient.Get();
+	UE_LOG(LogParty, Verbose, TEXT("Host connection failed for reservation beacon [%s]"), LocalReservationBeaconClient ? *LocalReservationBeaconClient->GetName() : TEXT(""));
 
 	// empty the queue, denying all requests
 	RejectAllPendingJoinRequests();
@@ -1591,7 +1594,7 @@ APartyBeaconClient* USocialParty::CreateReservationBeaconClient()
 	LastReservationBeaconClientNetDriverName = NAME_None;
 	ReservationBeaconClient = World->SpawnActor<APartyBeaconClient>(ReservationBeaconClientClass);
 	
-	return ReservationBeaconClient;
+	return ReservationBeaconClient.Get();
 }
 
 ASpectatorBeaconClient* USocialParty::CreateSpectatorBeaconClient()
@@ -1612,7 +1615,8 @@ void USocialParty::PumpApprovalQueue()
 	FPendingMemberApproval NextApproval;
 	if (PendingApprovals.Peek(NextApproval))
 	{
-		if (ensure(ReservationBeaconClient))
+		APartyBeaconClient* LocalReservationBeaconClient = ReservationBeaconClient.Get();
+		if (ensure(LocalReservationBeaconClient))
 		{
 			TArray<FPlayerReservation> PlayersToAdd;
 			PlayersToAdd.Reserve(NextApproval.Members.Num());
@@ -1635,7 +1639,7 @@ void USocialParty::PumpApprovalQueue()
 					NewPlayerRes.bAllowCrossplay = true;
 				}
 			}
-			ReservationBeaconClient->RequestReservationUpdate(GetPartyLeader()->GetPrimaryNetId(), PlayersToAdd);
+			LocalReservationBeaconClient->RequestReservationUpdate(GetPartyLeader()->GetPrimaryNetId(), PlayersToAdd);
 		}
 		else
 		{
@@ -1690,14 +1694,14 @@ void USocialParty::HandleReservationRequestComplete(EPartyReservationResult::Typ
 
 void USocialParty::CleanupReservationBeacon()
 {
-	if (ReservationBeaconClient)
+	if (APartyBeaconClient* LocalReservationBeaconClient = ReservationBeaconClient.Get())
 	{
-		UE_LOG(LogParty, Verbose, TEXT("Party reservation beacon cleanup while in state %s, pending approvals: %s"), ToString(ReservationBeaconClient->GetConnectionState()), !PendingApprovals.IsEmpty() ? TEXT("true") : TEXT("false"));
+		UE_LOG(LogParty, Verbose, TEXT("Party reservation beacon cleanup while in state %s, pending approvals: %s"), ToString(LocalReservationBeaconClient->GetConnectionState()), !PendingApprovals.IsEmpty() ? TEXT("true") : TEXT("false"));
 
-		LastReservationBeaconClientNetDriverName = ReservationBeaconClient->GetNetDriverName();
-		ReservationBeaconClient->OnHostConnectionFailure().Unbind();
-		ReservationBeaconClient->OnReservationRequestComplete().Unbind();
-		ReservationBeaconClient->DestroyBeacon();
+		LastReservationBeaconClientNetDriverName = LocalReservationBeaconClient->GetNetDriverName();
+		LocalReservationBeaconClient->OnHostConnectionFailure().Unbind();
+		LocalReservationBeaconClient->OnReservationRequestComplete().Unbind();
+		LocalReservationBeaconClient->DestroyBeacon();
 		ReservationBeaconClient = nullptr;
 	}
 }
diff --git a/Engine/Plugins/Online/OnlineFramework/Source/Party/Public/Party/SocialParty.h b/Engine/Plugins/Online/OnlineFramework/Source/Party/Public/Party/SocialParty.h
--- a/Engine/Plugins/Online/OnlineFramework/Source/Party/Public/Party/SocialParty.h
+++ b/Engine/Plugins/Online/OnlineFramework/Source/Party/Public/Party/SocialParty.h
@@ -289,7 +289,7 @@ protected:
 	void CleanupReservationBeacon();
 	APartyBeaconClient* CreateReservationBeaconClient();
 
-	APartyBeaconClient* GetReservationBeaconClient() const { return ReservationBeaconClient; }
+	APartyBeaconClient* GetReservationBeaconClient() const { return ReservationBeaconClient.Get(); }
 
 	/**
 	* Create a spectator beacon and connect to the server to get approval for new spectators
@@ -420,7 +420,7 @@ private:
 	
 	/** Reservation beacon client instance while getting approval for new party members*/
 	UPROPERTY()
-	APartyBeaconClient* ReservationBeaconClient = nullptr;
+	TWeakObjectPtr<APartyBeaconClient> ReservationBeaconClient = nullptr;
 	
 	/**
 	* Last known spectator beacon client net driver name
