Fix social party reservation requests not connecting to the reservation beacon.

diff --git a/Engine/Plugins/Online/OnlineFramework/Source/Party/Private/Party/SocialParty.cpp b/Engine/Plugins/Online/OnlineFramework/Source/Party/Private/Party/SocialParty.cpp
--- a/Engine/Plugins/Online/OnlineFramework/Source/Party/Private/Party/SocialParty.cpp
+++ b/Engine/Plugins/Online/OnlineFramework/Source/Party/Private/Party/SocialParty.cpp
@@ -1512,9 +1512,12 @@ void USocialParty::ConnectToReservationBeacon()
 					if (ensure(SessionInterface->GetResolvedConnectString(PartyGameSessionName, URL, NAME_BeaconPort)))
 					{
 						// Reconnect to the reservation beacon to maintain our place in the game (just until actual joined, holds place for all party members)
-						ReservationBeaconClient = World->SpawnActor<APartyBeaconClient>(ReservationBeaconClientClass);
+						LocalReservationBeaconClient = World->SpawnActor<APartyBeaconClient>(ReservationBeaconClientClass);
 						if (LocalReservationBeaconClient)
 						{
+							// Save as weak pointer.
+							ReservationBeaconClient = LocalReservationBeaconClient;
+
 							UE_LOG(LogParty, Verbose, TEXT("Party [%s] created reservation beacon [%s]."), *ToDebugString(), *LocalReservationBeaconClient->GetName());
 
 							LocalReservationBeaconClient->OnHostConnectionFailure().BindUObject(this, &USocialParty::HandleBeaconHostConnectionFailed);
