Converting SocialToolKit LocalPlayerOwner to weak pointer

diff --git a/Engine/Plugins/Online/OnlineFramework/Source/Party/Private/SocialToolkit.cpp b/Engine/Plugins/Online/OnlineFramework/Source/Party/Private/SocialToolkit.cpp
--- a/Engine/Plugins/Online/OnlineFramework/Source/Party/Private/SocialToolkit.cpp
+++ b/Engine/Plugins/Online/OnlineFramework/Source/Party/Private/SocialToolkit.cpp
@@ -324,8 +324,8 @@ USocialManager& USocialToolkit::GetSocialManager() const
 
 ULocalPlayer& USocialToolkit::GetOwningLocalPlayer() const
 {
-	check(LocalPlayerOwner);
-	return *LocalPlayerOwner;
+	check(LocalPlayerOwner.IsValid());
+	return *LocalPlayerOwner.Get();
 }
 
 USocialUser* USocialToolkit::FindUser(const FUniqueNetIdRepl& UserId) const
diff --git a/Engine/Plugins/Online/OnlineFramework/Source/Party/Public/SocialToolkit.h b/Engine/Plugins/Online/OnlineFramework/Source/Party/Public/SocialToolkit.h
--- a/Engine/Plugins/Online/OnlineFramework/Source/Party/Public/SocialToolkit.h
+++ b/Engine/Plugins/Online/OnlineFramework/Source/Party/Public/SocialToolkit.h
@@ -310,7 +310,7 @@ private:
 	TMap<FUniqueNetIdRepl, TWeakObjectPtr<USocialUser>> UsersBySubsystemIds;
 
 	UPROPERTY()
-	ULocalPlayer* LocalPlayerOwner = nullptr;
+	TWeakObjectPtr<ULocalPlayer> LocalPlayerOwner = nullptr;
 
 	UPROPERTY()
 	USocialChatManager* SocialChatManager;
