Refactor User List custom sort logic to use a list to allow dynamic ordering of custom sort values.

#rnx

diff --git a/Engine/Plugins/Online/OnlineFramework/Source/Party/Private/User/SocialUser.cpp b/Engine/Plugins/Online/OnlineFramework/Source/Party/Private/User/SocialUser.cpp
--- a/Engine/Plugins/Online/OnlineFramework/Source/Party/Private/User/SocialUser.cpp
+++ b/Engine/Plugins/Online/OnlineFramework/Source/Party/Private/User/SocialUser.cpp
@@ -1442,4 +1442,15 @@ FString USocialUser::SanitizePresenceString(FString InString) const
 	return InString;
 }
 
+void USocialUser::PopulateSortParameterList(TArray<int64>& OutSortParams) const
+{
+	const int32 TotalSortParams = 3;
+	OutSortParams.Reset(TotalSortParams);
+
+	// Parameters are prioritized by the first member added to the list being the highest priority
+	OutSortParams.Add(GetCustomSortValuePrimary());
+	OutSortParams.Add(GetCustomSortValueSecondary());
+	OutSortParams.Add(GetCustomSortValueTertiary());
+}
+
 #undef LOCTEXT_NAMESPACE
\ No newline at end of file
diff --git a/Engine/Plugins/Online/OnlineFramework/Source/Party/Private/User/SocialUserList.cpp b/Engine/Plugins/Online/OnlineFramework/Source/Party/Private/User/SocialUserList.cpp
--- a/Engine/Plugins/Online/OnlineFramework/Source/Party/Private/User/SocialUserList.cpp
+++ b/Engine/Plugins/Online/OnlineFramework/Source/Party/Private/User/SocialUserList.cpp
@@ -471,52 +471,45 @@ bool FSocialUserList::EvaluatePresenceFlag(bool bPresenceValue, ESocialUserState
 // encapsulates UserList sorting comparator and supporting data needed
 struct FUserSortData
 {
-	FUserSortData(USocialUser* InUser, EOnlinePresenceState::Type InStatus, bool InPlayingThisGame, FString InDisplayName, int64 InCustomSortValuePrimary, int64 InCustomSortValueSecondary, int64 InCustomSortValueTertiary)
+	FUserSortData(USocialUser* InUser, EOnlinePresenceState::Type InStatus, bool InPlayingThisGame, FString InDisplayName, TSharedRef<const TArray<int64>> InSortParams)
 		: User(InUser)
 		, OnlineStatus(InStatus)
 		, PlayingThisGame(InPlayingThisGame)
 		, DisplayName(MoveTemp(InDisplayName))
-		, CustomSortValuePrimary(InCustomSortValuePrimary)
-		, CustomSortValueSecondary(InCustomSortValueSecondary)
-		, CustomSortValueTertiary(InCustomSortValueTertiary) 
+		, SortParams(InSortParams)
 	{ }
 
 	USocialUser* User = nullptr;
 	EOnlinePresenceState::Type OnlineStatus;
 	bool PlayingThisGame;
 	FString DisplayName;
-	int64 CustomSortValuePrimary;
-	int64 CustomSortValueSecondary;
-	int64 CustomSortValueTertiary;
+	TSharedPtr<const TArray<int64>> SortParams;
 
 	bool operator<(const FUserSortData& OtherSortData) const
 	{
-		// Goes from if online, playing this game, then alphabetical
+		// Sorts in order of if online, playing this game, custom list of search parameters, then alphabetical
 		if (OnlineStatus == OtherSortData.OnlineStatus)
 		{
 			if (PlayingThisGame == OtherSortData.PlayingThisGame)
 			{
-				if (CustomSortValuePrimary == OtherSortData.CustomSortValuePrimary)
+				if (SortParams.IsValid() && OtherSortData.SortParams.IsValid() && SortParams->Num() == OtherSortData.SortParams->Num())
 				{
-					if (CustomSortValueSecondary == OtherSortData.CustomSortValueSecondary)
+					for (int32 SortIndex = 0; SortIndex < SortParams->Num() && SortIndex < OtherSortData.SortParams->Num(); ++SortIndex)
 					{
-						if (CustomSortValueTertiary == OtherSortData.CustomSortValueTertiary)
+						if ((*SortParams)[SortIndex] != (*OtherSortData.SortParams)[SortIndex])
 						{
-							return DisplayName < OtherSortData.DisplayName;
+							// Sort first by this parameter
+							return (*SortParams)[SortIndex] > (*OtherSortData.SortParams)[SortIndex];
 						}
-						else
-						{
-							return CustomSortValueTertiary > OtherSortData.CustomSortValueTertiary;
-						}
-					}
-					else
-					{
-						return CustomSortValueSecondary > OtherSortData.CustomSortValueSecondary;
 					}
+
+					// If we made it here, our sort params were all equal, so just sort alphabetically
+					return DisplayName < OtherSortData.DisplayName;
 				}
 				else
 				{
-					return CustomSortValuePrimary > OtherSortData.CustomSortValuePrimary;
+					// Our sort params don't match, so in this case just sort alphabetically
+					return DisplayName < OtherSortData.DisplayName;
 				}
 			}
 			else
@@ -618,13 +611,15 @@ void FSocialUserList::UpdateListInternal()
 				TArray<FUserSortData> SortedData;
 				SortedData.Reserve(NumUsers);
 
-				Algo::TransformIf(Users, SortedData, 
-					[](const TWeakObjectPtr<USocialUser> UserPtr) { return UserPtr.IsValid(); }, 
+				Algo::TransformIf(Users, SortedData,
+					[](const TWeakObjectPtr<USocialUser> UserPtr) { return UserPtr.IsValid(); },
 					[](const TWeakObjectPtr<USocialUser> UserPtr) -> FUserSortData
-				{
-					USocialUser* User = UserPtr.Get();
-					return FUserSortData(User, User->GetOnlineStatus(), User->IsPlayingThisGame(), User->GetDisplayName(), User->GetCustomSortValuePrimary(), User->GetCustomSortValueSecondary(), User->GetCustomSortValueTertiary());
-				});
+					{
+						USocialUser* User = UserPtr.Get();
+						TSharedRef<TArray<int64>> SortParams = MakeShared<TArray<int64>>();
+						User->PopulateSortParameterList(SortParams.Get());
+						return FUserSortData(User, User->GetOnlineStatus(), User->IsPlayingThisGame(), User->GetDisplayName(), SortParams);
+					});
 
 				Algo::Sort(SortedData);
 
diff --git a/Engine/Plugins/Online/OnlineFramework/Source/Party/Public/User/SocialUser.h b/Engine/Plugins/Online/OnlineFramework/Source/Party/Public/User/SocialUser.h
--- a/Engine/Plugins/Online/OnlineFramework/Source/Party/Public/User/SocialUser.h
+++ b/Engine/Plugins/Online/OnlineFramework/Source/Party/Public/User/SocialUser.h
@@ -87,6 +87,9 @@ public:
 	virtual int64 GetCustomSortValueSecondary() const { return 0; }
 	virtual int64 GetCustomSortValueTertiary() const { return 0; }
 
+	/** Populate list with sort values in order of priority */
+	virtual void PopulateSortParameterList(TArray<int64>& OutSortParams) const;
+
 	bool SetUserLocalAttribute(ESocialSubsystem SubsystemType, const FString& AttrName, const FString& AttrValue);
 	bool GetUserAttribute(ESocialSubsystem SubsystemType, const FString& AttrName, FString& OutAttrValue) const;
 
