Add helper method to get a social platform description by online subsystem name

diff --git a/Engine/Plugins/Online/OnlineFramework/Source/Party/Private/Party/PartyPlatformSessionMonitor.cpp b/Engine/Plugins/Online/OnlineFramework/Source/Party/Private/Party/PartyPlatformSessionMonitor.cpp
--- a/Engine/Plugins/Online/OnlineFramework/Source/Party/Private/Party/PartyPlatformSessionMonitor.cpp
+++ b/Engine/Plugins/Online/OnlineFramework/Source/Party/Private/Party/PartyPlatformSessionMonitor.cpp
@@ -109,11 +109,7 @@ static FAutoConsoleVariableRef CVar_EstablishSessionRetryDelay(
 
 static const FSocialPlatformDescription* FindPlatformDescriptionByOss(FName OssName)
 {
-	auto FindPlatformDescriptionByOssName = [OssName](const FSocialPlatformDescription& TestPlatformDescription)
-	{
-		return TestPlatformDescription.OnlineSubsystem == OssName;
-	};
-	return USocialSettings::GetSocialPlatformDescriptions().FindByPredicate(FindPlatformDescriptionByOssName);
+	return USocialSettings::GetSocialPlatformDescriptionForOnlineSubsystem(OssName);
 }
 
 bool FPartyPlatformSessionManager::DoesOssNeedPartySession(FName OssName)
diff --git a/Engine/Plugins/Online/OnlineFramework/Source/Party/Private/SocialSettings.cpp b/Engine/Plugins/Online/OnlineFramework/Source/Party/Private/SocialSettings.cpp
--- a/Engine/Plugins/Online/OnlineFramework/Source/Party/Private/SocialSettings.cpp
+++ b/Engine/Plugins/Online/OnlineFramework/Source/Party/Private/SocialSettings.cpp
@@ -97,3 +97,11 @@ const TArray<FSocialPlatformDescription>& USocialSettings::GetSocialPlatformDesc
 	const USocialSettings& SettingsCDO = *GetDefault<USocialSettings>();
 	return SettingsCDO.SocialPlatformDescriptions;
 }
+
+const FSocialPlatformDescription* USocialSettings::GetSocialPlatformDescriptionForOnlineSubsystem(const FName& OnlineSubsystemName)
+{
+	return GetSocialPlatformDescriptions().FindByPredicate([&OnlineSubsystemName](const FSocialPlatformDescription& Candidate)
+	{
+		return Candidate.OnlineSubsystem == OnlineSubsystemName;
+	});
+}
diff --git a/Engine/Plugins/Online/OnlineFramework/Source/Party/Private/SocialToolkit.cpp b/Engine/Plugins/Online/OnlineFramework/Source/Party/Private/SocialToolkit.cpp
--- a/Engine/Plugins/Online/OnlineFramework/Source/Party/Private/SocialToolkit.cpp
+++ b/Engine/Plugins/Online/OnlineFramework/Source/Party/Private/SocialToolkit.cpp
@@ -93,11 +93,8 @@ public:
 		{
 			// The external mappings will always be checked on the primary OSS, so we use the passed-in OSS as the target we want to map to
 			IOnlineSubsystem* OSS = GetOSS();
-			auto FindPlatformDescriptionByOssName = [OSS](const FSocialPlatformDescription& TestPlatformDescription)
-			{
-				return TestPlatformDescription.OnlineSubsystem == OSS->GetSubsystemName();
-			};
-			const FSocialPlatformDescription* PlatformDescription = OSS ? USocialSettings::GetSocialPlatformDescriptions().FindByPredicate(FindPlatformDescriptionByOssName) : nullptr;
+			check(OSS);
+			const FSocialPlatformDescription* PlatformDescription = OSS ? USocialSettings::GetSocialPlatformDescriptionForOnlineSubsystem(OSS->GetSubsystemName()) : nullptr;
 			IOnlineUserPtr PrimaryUserInterface = Toolkit->GetSocialOss(ESocialSubsystem::Primary)->GetUserInterface();
 			if (ensure(PlatformDescription && PrimaryUserInterface))
 			{
@@ -429,11 +426,8 @@ void USocialToolkit::QueueUserDependentActionInternal(const FUniqueNetIdRepl& Su
 			IOnlineUserPtr UserInterface = Online::GetUserInterfaceChecked(GetWorld(), USocialManager::GetSocialOssName(ESocialSubsystem::Primary));
 
 			IOnlineSubsystem* OSS = GetSocialOss(SubsystemType);
-			auto FindPlatformDescriptionByOssName = [OSS](const FSocialPlatformDescription& TestPlatformDescription)
-			{
-				return TestPlatformDescription.OnlineSubsystem == OSS->GetSubsystemName();
-			};
-			const FSocialPlatformDescription* PlatformDescription = OSS ? USocialSettings::GetSocialPlatformDescriptions().FindByPredicate(FindPlatformDescriptionByOssName) : nullptr;
+			check(OSS);
+			const FSocialPlatformDescription* PlatformDescription = OSS ? USocialSettings::GetSocialPlatformDescriptionForOnlineSubsystem(OSS->GetSubsystemName()) : nullptr;
 			
 			FExternalIdQueryOptions QueryOptions;
 			QueryOptions.AuthType = PlatformDescription ? PlatformDescription->ExternalAccountType : FString();
diff --git a/Engine/Plugins/Online/OnlineFramework/Source/Party/Private/User/SocialUser.cpp b/Engine/Plugins/Online/OnlineFramework/Source/Party/Private/User/SocialUser.cpp
--- a/Engine/Plugins/Online/OnlineFramework/Source/Party/Private/User/SocialUser.cpp
+++ b/Engine/Plugins/Online/OnlineFramework/Source/Party/Private/User/SocialUser.cpp
@@ -1378,11 +1378,7 @@ void USocialUser::SetUserInfo(ESocialSubsystem SubsystemType, const TSharedRef<F
 			{
 				if (IOnlineSubsystem* MissingOSS = GetOwningToolkit().GetSocialOss(Subsystem))
 				{
-					auto FindPlatformDescriptionByOssName = [MissingOSS](const FSocialPlatformDescription& TestPlatformDescription)
-					{
-						return TestPlatformDescription.OnlineSubsystem == MissingOSS->GetSubsystemName();
-					};
-					if (const FSocialPlatformDescription* PlatformDescription = USocialSettings::GetSocialPlatformDescriptions().FindByPredicate(FindPlatformDescriptionByOssName))
+					if (const FSocialPlatformDescription* PlatformDescription = USocialSettings::GetSocialPlatformDescriptionForOnlineSubsystem(MissingOSS->GetSubsystemName()))
 					{
 						if (!PlatformDescription->ExternalAccountType.IsEmpty())
 						{
diff --git a/Engine/Plugins/Online/OnlineFramework/Source/Party/Public/SocialSettings.h b/Engine/Plugins/Online/OnlineFramework/Source/Party/Public/SocialSettings.h
--- a/Engine/Plugins/Online/OnlineFramework/Source/Party/Public/SocialSettings.h
+++ b/Engine/Plugins/Online/OnlineFramework/Source/Party/Public/SocialSettings.h
@@ -65,6 +65,12 @@ public:
 	static int32 GetMinNicknameLength();
 	static int32 GetMaxNicknameLength();
 	static const TArray<FSocialPlatformDescription>& GetSocialPlatformDescriptions();
+	/**
+	 * Get a platform description (from GetSocialPlatformDescriptions) for a specific OnlineSubsystem.
+	 * @param OnlineSubsystemName the online subsystem name to search for
+	 * @return the social platform description for that online subsystem. May return null if it is not found.
+	 */
+	static const FSocialPlatformDescription* GetSocialPlatformDescriptionForOnlineSubsystem(const FName& OnlineSubsystemName);
 
 private:
 	/**
