Fixed KPK error by removing stale ClientActor references from LobbyBeaconPlayerState when the LobbyBeaconClient ends play
#tests Was able to consistently repro this on the server by having a client go from a lobby to in-game. Confirmed that error doesn't occur with these changes.

diff --git a/Engine/Plugins/Online/OnlineFramework/Source/Lobby/Private/LobbyBeaconClient.cpp b/Engine/Plugins/Online/OnlineFramework/Source/Lobby/Private/LobbyBeaconClient.cpp
--- a/Engine/Plugins/Online/OnlineFramework/Source/Lobby/Private/LobbyBeaconClient.cpp
+++ b/Engine/Plugins/Online/OnlineFramework/Source/Lobby/Private/LobbyBeaconClient.cpp
@@ -29,6 +29,24 @@ void ALobbyBeaconClient::GetLifetimeReplicatedProps(TArray< FLifetimeProperty >
 	DOREPLIFETIME(ALobbyBeaconClient, PlayerState);
 }
 
+void ALobbyBeaconClient::EndPlay(EEndPlayReason::Type Reason)
+{
+	Super::EndPlay(Reason);
+
+	if (PlayerState)
+	{
+		if (PlayerState->ClientActor == this)
+		{
+			PlayerState->ClientActor = nullptr;
+		}
+
+		if (PlayerState->GetOwner() == this)
+		{
+			PlayerState->SetOwner(nullptr);
+		}
+	}
+}
+
 void ALobbyBeaconClient::OnConnected()
 {
 	UE_LOG(LogLobbyBeacon, Verbose, TEXT("Lobby beacon connection established."));
diff --git a/Engine/Plugins/Online/OnlineFramework/Source/Lobby/Public/LobbyBeaconClient.h b/Engine/Plugins/Online/OnlineFramework/Source/Lobby/Public/LobbyBeaconClient.h
--- a/Engine/Plugins/Online/OnlineFramework/Source/Lobby/Public/LobbyBeaconClient.h
+++ b/Engine/Plugins/Online/OnlineFramework/Source/Lobby/Public/LobbyBeaconClient.h
@@ -65,6 +65,9 @@ class LOBBY_API ALobbyBeaconClient : public AOnlineBeaconClient
 	/** Player state associated with this beacon (@todo not splitscreen safe) */
 	UPROPERTY(Replicated)
 	ALobbyBeaconPlayerState* PlayerState;
+
+	/** Clear out any references to this in PlayerState */
+	virtual void EndPlay(EEndPlayReason::Type Reason) override;
 	
 	//~ Begin AOnlineBeaconClient Interface
 	virtual void OnConnected() override;
