Add new IsAnyInboundFriendshipPending to USocialUser

#tests in PIE, used 2 accounts A and B and verified that it only returned true for B if, A had sent a friend invite, and not when B did send the invite.

diff --git a/Engine/Plugins/Online/OnlineFramework/Source/Party/Private/User/SocialUser.cpp b/Engine/Plugins/Online/OnlineFramework/Source/Party/Private/User/SocialUser.cpp
--- a/Engine/Plugins/Online/OnlineFramework/Source/Party/Private/User/SocialUser.cpp
+++ b/Engine/Plugins/Online/OnlineFramework/Source/Party/Private/User/SocialUser.cpp
@@ -514,6 +514,22 @@ bool USocialUser::IsFriendshipPending(ESocialSubsystem SubsystemType) const
 	return FriendInviteStatus == EInviteStatus::PendingInbound || FriendInviteStatus == EInviteStatus::PendingOutbound;
 }
 
+bool USocialUser::IsAnyInboundFriendshipPending() const
+{
+	for (const TPair<ESocialSubsystem, FSubsystemUserInfo>& SubsystemInfoPair : SubsystemInfoByType)
+	{
+		if (const FOnlineFriend* OnlineFriend = SubsystemInfoPair.Value.FriendInfo.Pin().Get())
+		{
+			const EInviteStatus::Type FriendInviteStatus = OnlineFriend->GetInviteStatus();
+			if (FriendInviteStatus == EInviteStatus::PendingInbound)
+			{
+				return true;
+			}
+		}
+	}
+	return false;
+}
+
 const FOnlineUserPresence* USocialUser::GetFriendPresenceInfo(ESocialSubsystem SubsystemType) const
 {
 	const FSubsystemUserInfo* SubsystemInfo = SubsystemInfoByType.Find(SubsystemType);
diff --git a/Engine/Plugins/Online/OnlineFramework/Source/Party/Public/User/SocialUser.h b/Engine/Plugins/Online/OnlineFramework/Source/Party/Public/User/SocialUser.h
--- a/Engine/Plugins/Online/OnlineFramework/Source/Party/Public/User/SocialUser.h
+++ b/Engine/Plugins/Online/OnlineFramework/Source/Party/Public/User/SocialUser.h
@@ -60,6 +60,7 @@ public:
 	bool IsFriend() const;
 	bool IsFriend(ESocialSubsystem SubsystemType) const;
 	bool IsFriendshipPending(ESocialSubsystem SubsystemType) const;
+	bool IsAnyInboundFriendshipPending() const;
 	const FOnlineUserPresence* GetFriendPresenceInfo(ESocialSubsystem SubsystemType) const;
 	FDateTime GetFriendshipCreationDate() const;
 	virtual FDateTime GetLastOnlineDate() const;
