Move join in progress flow to OnlineFramework using party member data and deprecate old OSS flow.

diff --git a/Engine/Plugins/Online/OnlineFramework/Source/Party/Private/Party/PartyMember.cpp b/Engine/Plugins/Online/OnlineFramework/Source/Party/Private/Party/PartyMember.cpp
--- a/Engine/Plugins/Online/OnlineFramework/Source/Party/Private/Party/PartyMember.cpp
+++ b/Engine/Plugins/Online/OnlineFramework/Source/Party/Private/Party/PartyMember.cpp
@@ -31,6 +31,8 @@ void FPartyMemberRepData::CompareAgainst(const FOnlinePartyRepDataBase& OldData)
 	ComparePlatformDataUniqueId(TypedOldData);
 	ComparePlatformDataSessionId(TypedOldData);
 	CompareCrossplayPreference(TypedOldData);
+	CompareJoinInProgressDataRequest(TypedOldData);
+	CompareJoinInProgressDataResponses(TypedOldData);
 }
 
 const USocialParty* FPartyMemberRepData::GetOwnerParty() const
diff --git a/Engine/Plugins/Online/OnlineFramework/Source/Party/Private/Party/SocialParty.cpp b/Engine/Plugins/Online/OnlineFramework/Source/Party/Private/Party/SocialParty.cpp
--- a/Engine/Plugins/Online/OnlineFramework/Source/Party/Private/Party/SocialParty.cpp
+++ b/Engine/Plugins/Online/OnlineFramework/Source/Party/Private/Party/SocialParty.cpp
@@ -133,6 +133,7 @@ ECrossplayPreference GetCrossplayPreferenceFromJoinData(const FOnlinePartyData&
 	return ECrossplayPreference::NoSelection;
 }
 
+// DEPRECATED - Use the new join in progress flow with USocialParty::RequestJoinInProgress.
 FPartyJoinApproval USocialParty::EvaluateJIPRequest(const FUniqueNetId& PlayerId) const
 {
 	FPartyJoinApproval JoinApproval;
@@ -473,17 +474,21 @@ void USocialParty::InitializePartyInternal()
 	PartyInterface->AddOnPartyConfigChangedDelegate_Handle(FOnPartyConfigChangedDelegate::CreateUObject(this, &USocialParty::HandlePartyConfigChanged));
 	PartyInterface->AddOnPartyDataReceivedDelegate_Handle(FOnPartyDataReceivedDelegate::CreateUObject(this, &USocialParty::HandlePartyDataReceived));
 	PartyInterface->AddOnPartyJoinRequestReceivedDelegate_Handle(FOnPartyJoinRequestReceivedDelegate::CreateUObject(this, &USocialParty::HandlePartyJoinRequestReceived));
-	PartyInterface->AddOnPartyJIPRequestReceivedDelegate_Handle(FOnPartyJIPRequestReceivedDelegate::CreateUObject(this, &USocialParty::HandlePartyJIPRequestReceived));
 	PartyInterface->AddOnQueryPartyJoinabilityReceivedDelegate_Handle(FOnQueryPartyJoinabilityReceivedDelegate::CreateUObject(this, &USocialParty::HandleJoinabilityQueryReceived));
 	PartyInterface->AddOnPartyExitedDelegate_Handle(FOnPartyExitedDelegate::CreateUObject(this, &USocialParty::HandlePartyLeft));
 	PartyInterface->AddOnPartyStateChangedDelegate_Handle(FOnPartyStateChangedDelegate::CreateUObject(this, &USocialParty::HandlePartyStateChanged));
 
 	PartyInterface->AddOnPartyMemberJoinedDelegate_Handle(FOnPartyMemberJoinedDelegate::CreateUObject(this, &USocialParty::HandlePartyMemberJoined));
-	PartyInterface->AddOnPartyJIPResponseDelegate_Handle(FOnPartyJIPResponseDelegate::CreateUObject(this, &USocialParty::HandlePartyMemberJIP));
 	PartyInterface->AddOnPartyMemberDataReceivedDelegate_Handle(FOnPartyMemberDataReceivedDelegate::CreateUObject(this, &USocialParty::HandlePartyMemberDataReceived));
 	PartyInterface->AddOnPartyMemberPromotedDelegate_Handle(FOnPartyMemberPromotedDelegate::CreateUObject(this, &USocialParty::HandlePartyMemberPromoted));
 	PartyInterface->AddOnPartyMemberExitedDelegate_Handle(FOnPartyMemberExitedDelegate::CreateUObject(this, &USocialParty::HandlePartyMemberExited));
 	PartyInterface->AddOnPartySystemStateChangeDelegate_Handle(FOnPartySystemStateChangeDelegate::CreateUObject(this, &USocialParty::HandlePartySystemStateChange));
+
+	PRAGMA_DISABLE_DEPRECATION_WARNINGS
+	PartyInterface->AddOnPartyJIPRequestReceivedDelegate_Handle(FOnPartyJIPRequestReceivedDelegate::CreateUObject(this, &USocialParty::HandlePartyJIPRequestReceived));
+	PartyInterface->AddOnPartyJIPResponseDelegate_Handle(FOnPartyJIPResponseDelegate::CreateUObject(this, &USocialParty::HandlePartyMemberJIP));
+	PRAGMA_ENABLE_DEPRECATION_WARNINGS
+
 	// Create a UPartyMember for every existing member on the OSS party
 	TArray<FOnlinePartyMemberConstRef> OssPartyMembers;
 	PartyInterface->GetPartyMembers(*OwningLocalUserId, GetPartyId(), OssPartyMembers);
@@ -802,6 +807,7 @@ void USocialParty::RemovePlayerFromReservationBeacon(const FUniqueNetId& LocalUs
 	}
 }
 
+// DEPRECATED - Use the new join in progress flow with USocialParty::RequestJoinInProgress.
 void USocialParty::HandlePartyJIPRequestReceived(const FUniqueNetId& LocalUserId, const FOnlinePartyId& PartyId, const FUniqueNetId& SenderId)
 {
 	if (!IsLocalPlayerPartyLeader() || PartyId != GetPartyId())
@@ -810,8 +816,10 @@ void USocialParty::HandlePartyJIPRequestReceived(const FUniqueNetId& LocalUserId
 	}
 
 	IOnlinePartyPtr PartyInterface = Online::GetPartyInterfaceChecked(GetWorld());
-	
+
+	PRAGMA_DISABLE_DEPRECATION_WARNINGS
 	FPartyJoinApproval JoinApproval = EvaluateJIPRequest(SenderId);
+	PRAGMA_ENABLE_DEPRECATION_WARNINGS
 
 	if (JoinApproval.GetApprovalAction() == EApprovalAction::Enqueue ||
 		JoinApproval.GetApprovalAction() == EApprovalAction::EnqueueAndStartBeacon)
@@ -845,11 +853,12 @@ void USocialParty::HandlePartyJIPRequestReceived(const FUniqueNetId& LocalUserId
 		const bool bIsApproved = JoinApproval.CanJoin();
 		UE_LOG(LogParty, Verbose, TEXT("[%s] Responding to approval request for %s with %s"), *PartyId.ToString(), *SenderId.ToString(), bIsApproved ? TEXT("approved") : TEXT("denied"));
 
+		PRAGMA_DISABLE_DEPRECATION_WARNINGS
 		PartyInterface->ApproveJIPRequest(LocalUserId, PartyId, SenderId, bIsApproved, JoinApproval.GetDenialReason());
+		PRAGMA_ENABLE_DEPRECATION_WARNINGS
 	}
 }
 
-
 void USocialParty::HandleJoinabilityQueryReceived(const FUniqueNetId& LocalUserId, const FOnlinePartyId& PartyId, const IOnlinePartyPendingJoinRequestInfo& JoinRequestInfo)
 {
 	if (PartyId == GetPartyId()) 
@@ -941,6 +950,7 @@ void USocialParty::HandlePartyMemberJoined(const FUniqueNetId& LocalUserId, cons
 	}
 }
 
+// DEPRECATED - Use the new join in progress flow with USocialParty::RequestJoinInProgress.
 void USocialParty::HandlePartyMemberJIP(const FUniqueNetId& LocalUserId, const FOnlinePartyId& PartyId, bool Success, int32 DeniedResultCode)
 {
 	if (PartyId == GetPartyId())
@@ -1006,9 +1016,11 @@ void USocialParty::HandleMemberInitialized(UPartyMember* Member)
 	{
 		Member->GetRepData().OnPlatformDataUniqueIdChanged().AddUObject(this, &USocialParty::HandleMemberPlatformUniqueIdChanged, Member);
 		Member->GetRepData().OnPlatformDataSessionIdChanged().AddUObject(this, &USocialParty::HandleMemberSessionIdChanged, Member);
-
 		HandleMemberPlatformUniqueIdChanged(Member->GetRepData().GetPlatformDataUniqueId(), Member);
 	}
+
+	Member->GetRepData().OnJoinInProgressDataRequestChanged().AddUObject(this, &USocialParty::HandleJoinInProgressDataRequestChanged, Member);
+	Member->GetRepData().OnJoinInProgressDataResponsesChanged().AddUObject(this, &USocialParty::HandleJoinInProgressDataResponsesChanged, Member);
 }
 
 void USocialParty::HandleMemberPlatformUniqueIdChanged(const FUniqueNetIdRepl& NewPlatformUniqueId, UPartyMember* Member)
@@ -1569,7 +1581,11 @@ void USocialParty::RejectAllPendingJoinRequests()
 		UE_LOG(LogParty, Verbose, TEXT("[%s] Responding to approval request for %s with denied"), *PartyId.ToString(), *PrimaryJoiningUserId.ToDebugString());
 		if (PendingApproval.bIsJIPApproval)
 		{
+			PRAGMA_DISABLE_DEPRECATION_WARNINGS
 			PartyInterface->ApproveJIPRequest(*PendingApproval.RecipientId, PartyId, PrimaryJoiningUserId, false, (int32)EPartyJoinDenialReason::Busy);
+			PRAGMA_ENABLE_DEPRECATION_WARNINGS
+
+			RespondToJoinInProgressRequest(PendingApproval, EPartyJoinDenialReason::Busy);
 		}
 		else
 		{
@@ -1673,7 +1689,11 @@ void USocialParty::HandleReservationRequestComplete(EPartyReservationResult::Typ
 			if (PendingApproval.bIsJIPApproval)
 			{
 				// This player is already in our party. ApproveJIPRequest
+				PRAGMA_DISABLE_DEPRECATION_WARNINGS
 				PartyInterface->ApproveJIPRequest(*PendingApproval.RecipientId, GetPartyId(), *PendingApproval.Members[0].MemberId, bReservationApproved, DenialReason);
+				PRAGMA_ENABLE_DEPRECATION_WARNINGS
+
+				RespondToJoinInProgressRequest(PendingApproval, DenialReason.GetReason());
 			}
 			else if (PendingApproval.bIsPlayerRemoval)
 			{
@@ -1903,4 +1923,190 @@ bool USocialParty::ShouldAlwaysJoinPlatformSession(const FSessionId& SessionId)
 void USocialParty::JoinSessionCompleteAnalytics(const FSessionId& SessionId, const FString& JoinBootableGroupSessionResult)
 {
 	// Work is to be done in the override
+}
+
+void USocialParty::RequestJoinInProgress(const UPartyMember& TargetMember, const FOnRequestJoinInProgressComplete& CompletionDelegate)
+{
+	// Only allow one join attempt at a time
+	if (RequestJoinInProgressComplete.IsSet())
+	{
+		UE_LOG(LogParty, Warning, TEXT("RequestJoinInProgress: Request already in progress"));
+		GetWorld()->GetTimerManager().SetTimerForNextTick(FTimerDelegate::CreateLambda([CompletionDelegate]() {
+			CompletionDelegate.ExecuteIfBound(EPartyJoinDenialReason::Busy);
+		}));
+		return;
+	}
+
+	RequestJoinInProgressComplete = CompletionDelegate;
+	FPartyMemberJoinInProgressRequest Request;
+	Request.Target = TargetMember.GetPrimaryNetId();
+	Request.Time = FDateTime::UtcNow().ToUnixTimestamp();
+
+	UE_LOG(LogParty, Verbose, TEXT("RequestJoinInProgress: Sending request Target=%s Time=%d"), *Request.Target.ToDebugString(), Request.Time);
+	GetOwningLocalMember().GetMutableRepData().SetJoinInProgressDataRequest(Request);
+	RunJoinInProgressTimer();
+}
+
+void USocialParty::HandleJoinInProgressDataRequestChanged(const FPartyMemberJoinInProgressRequest& Request, UPartyMember* Member)
+{
+	if (Request.Time == 0 || !IsLocalPlayerPartyLeader())
+	{
+		// Ignore if this is not an active request or if we're not the party leader.
+		return;
+	}
+
+	const UPartyMember* TargetPartyMember = GetPartyMember(Request.Target);
+	if (!TargetPartyMember)
+	{
+		UE_LOG(LogParty, Warning, TEXT("HandleJoinInProgressDataRequestChanged: Could not find member for request Target=%s Time=%d"), *Request.Target.ToDebugString(), Request.Time);
+		return;
+	}
+
+	if (!TargetPartyMember->IsLocalPlayer())
+	{
+		// Request was not sent to us.
+		return;
+	}
+
+	UE_LOG(LogParty, Verbose, TEXT("HandleJoinInProgressDataRequestChanged: Received request Requester=%s Time=%d"), *Member->GetPrimaryNetId().ToDebugString(), Request.Time);
+	FPendingMemberApproval PendingApproval;
+	PendingApproval.RecipientId = OwningLocalUserId;
+	PendingApproval.Members.Emplace(Member->GetPrimaryNetId(), Member->GetRepData().GetPlatformDataPlatform());
+	PendingApproval.bIsJIPApproval = true;
+	PendingApproval.JoinInProgressRequestTime = Request.Time;
+	PendingApprovals.Enqueue(MoveTemp(PendingApproval));
+	if (!ReservationBeaconClient.Get())
+	{
+		ConnectToReservationBeacon();
+	}
+}
+
+void USocialParty::RespondToJoinInProgressRequest(const FPendingMemberApproval& PendingApproval, const EPartyJoinDenialReason DenialReason)
+{
+	if (PendingApproval.Members.IsEmpty() || PendingApproval.JoinInProgressRequestTime == 0)
+	{
+		return;
+	}
+
+	const UPartyMember* RequestingMember = GetPartyMember(PendingApproval.Members[0].MemberId);
+	if (!RequestingMember)
+	{
+		UE_LOG(LogParty, Warning, TEXT("RespondToJoinInProgressRequest: Could not find member for approval MemberId=%s"), *PendingApproval.Members[0].MemberId.ToDebugString());
+		return;
+	}
+
+	FPartyMemberJoinInProgressResponse Response;
+	Response.Requester = RequestingMember->GetPrimaryNetId();
+	Response.Time = PendingApproval.JoinInProgressRequestTime;
+	Response.DenialReason = static_cast<uint8>(DenialReason);
+
+	UE_LOG(LogParty, Verbose, TEXT("RespondToJoinInProgressRequest: Sending response Requester=%s Time=%d DenialReason=%s"),
+		*Response.Requester.ToDebugString(), Response.Time, ToString(DenialReason));
+	TArray<FPartyMemberJoinInProgressResponse> Responses = GetOwningLocalMember().GetRepData().GetJoinInProgressDataResponses();
+	Responses.Add(Response);
+	GetOwningLocalMember().GetMutableRepData().SetJoinInProgressDataResponses(Responses);
+	RunJoinInProgressTimer();
+}
+
+void USocialParty::HandleJoinInProgressDataResponsesChanged(const TArray<FPartyMemberJoinInProgressResponse>& Responses, UPartyMember* Member)
+{
+	if (!RequestJoinInProgressComplete.IsSet())
+	{
+		// Skip if we're not waiting for a response.
+		return;
+	}
+
+	const FPartyMemberJoinInProgressRequest Request = GetOwningLocalMember().GetRepData().GetJoinInProgressDataRequest();
+
+	for (const FPartyMemberJoinInProgressResponse& Response : Responses)
+	{
+		if (Response.Time != Request.Time)
+		{
+			// Response was not for us.
+			continue;
+		}
+
+		const UPartyMember* RequestingMember = GetPartyMember(Response.Requester);
+		const EPartyJoinDenialReason DenialReason = static_cast<EPartyJoinDenialReason>(Response.DenialReason);
+
+		if (!RequestingMember)
+		{
+			UE_LOG(LogParty, Warning, TEXT("HandleJoinInProgressDataResponsesChanged: Could not find member for response Requester=%s Time=%d DenialReason=%s"),
+				*Response.Requester.ToDebugString(), Response.Time, ToString(DenialReason));
+			continue;
+		}
+
+		if (!RequestingMember->IsLocalPlayer())
+		{
+			// Response was not for us.
+			continue;
+		}
+
+		// Responses are ordered newest first, so use the first one we find.
+		const int64 Now = FDateTime::UtcNow().ToUnixTimestamp();
+		UE_LOG(LogParty, Verbose, TEXT("HandleJoinInProgressDataResponsesChanged: Received response Requester=%s Time=%d DenialReason=%s ResponseTime=%d"),
+			*Response.Requester.ToDebugString(), Response.Time, ToString(DenialReason), Now - Response.Time);
+		CallJoinInProgressComplete(DenialReason);
+		break;
+	}
+}
+
+void USocialParty::CallJoinInProgressComplete(const EPartyJoinDenialReason DenialReason)
+{
+	if (RequestJoinInProgressComplete.IsSet())
+	{
+		FOnRequestJoinInProgressComplete CompletionDelegate = MoveTemp(RequestJoinInProgressComplete.GetValue());
+		RequestJoinInProgressComplete.Reset();
+		CompletionDelegate.ExecuteIfBound(DenialReason);
+	}
+}
+
+void USocialParty::RunJoinInProgressTimer()
+{
+	UE_LOG(LogParty, VeryVerbose, TEXT("RunJoinInProgressTimer: Checking for stale data"));
+
+	const int64 Now = FDateTime::UtcNow().ToUnixTimestamp();
+	int64 NextTimer = 0;
+
+	FPartyMemberJoinInProgressRequest Request = GetOwningLocalMember().GetRepData().GetJoinInProgressDataRequest();
+	if (Request.Time > 0)
+	{
+		const int64 Expires = Request.Time + JoinInProgressRequestTimeout;
+		if (Expires <= Now)
+		{
+			UE_LOG(LogParty, Verbose, TEXT("RunJoinInProgressTimer: Removing request data"));
+			CallJoinInProgressComplete(EPartyJoinDenialReason::JoinAttemptAborted);
+			Request.Target = FUniqueNetIdRepl::Invalid();
+			Request.Time = 0;
+			GetOwningLocalMember().GetMutableRepData().SetJoinInProgressDataRequest(Request);
+		}
+		else
+		{
+			NextTimer = Expires - Now;
+		}
+	}
+
+	const TArray<FPartyMemberJoinInProgressResponse>& Responses = GetOwningLocalMember().GetRepData().GetJoinInProgressDataResponses();
+	TArray<FPartyMemberJoinInProgressResponse> ResponsesToKeep;
+	for (const FPartyMemberJoinInProgressResponse& Response : Responses)
+	{
+		const int64 Expires = Response.Time + JoinInProgressResponseTimeout;
+		if (Expires > Now)
+		{
+			ResponsesToKeep.Add(Response);
+			NextTimer = NextTimer ? FMath::Min(NextTimer, Expires - Now) : Expires - Now;
+		}
+	}
+
+	if (Responses.Num() != ResponsesToKeep.Num())
+	{
+		UE_LOG(LogParty, Verbose, TEXT("RunJoinInProgressTimer: Removing response data, %d remaining"), ResponsesToKeep.Num());
+		GetOwningLocalMember().GetMutableRepData().SetJoinInProgressDataResponses(ResponsesToKeep);
+	}
+
+	if (NextTimer > 0)
+	{
+		UE_LOG(LogParty, Verbose, TEXT("RunJoinInProgressTimer: Running again in %d seconds"), NextTimer);
+		GetWorld()->GetTimerManager().SetTimer(JoinInProgressTimerHandle, this, &USocialParty::RunJoinInProgressTimer, static_cast<float>(NextTimer));
+	}
 }
\ No newline at end of file
diff --git a/Engine/Plugins/Online/OnlineFramework/Source/Party/Public/Party/PartyMember.h b/Engine/Plugins/Online/OnlineFramework/Source/Party/Public/Party/PartyMember.h
--- a/Engine/Plugins/Online/OnlineFramework/Source/Party/Public/Party/PartyMember.h
+++ b/Engine/Plugins/Online/OnlineFramework/Source/Party/Public/Party/PartyMember.h
@@ -39,6 +39,71 @@ public:
 	FString SessionId;
 };
 
+/** Join in progress request. Represents a request from a local party member to a remote party member to acquire a reservation for the session the remote party member is in. */
+USTRUCT()
+struct FPartyMemberJoinInProgressRequest
+{
+	GENERATED_BODY()
+
+public:
+	bool operator==(const FPartyMemberJoinInProgressRequest& Other) const { return Target == Other.Target && Time == Other.Time; }
+	bool operator!=(const FPartyMemberJoinInProgressRequest& Other) const { return !operator==(Other); }
+
+	/** Remote member we want to join. */
+	UPROPERTY()
+	FUniqueNetIdRepl Target;
+
+	/** Time the request was made. */
+	UPROPERTY()
+	int64 Time = 0;
+};
+
+/** Join in progress response. Represents a response from a local party member to a remote party member that requested to join in progress. */
+USTRUCT()
+struct FPartyMemberJoinInProgressResponse
+{
+	GENERATED_BODY()
+
+public:
+	bool operator==(const FPartyMemberJoinInProgressResponse& Other) const { return Requester == Other.Requester && Time == Other.Time; }
+	bool operator!=(const FPartyMemberJoinInProgressResponse& Other) const { return !operator==(Other); }
+
+	/** Remote member that this response is for. */
+	UPROPERTY()
+	FUniqueNetIdRepl Requester;
+
+	/** Time the request was made. Matches FPartyMemberJoinInProgressRequest::Time */
+	UPROPERTY()
+	int64 Time = 0;
+	
+	/**
+	 * Result of session reservation attempt.
+	 * @see EPartyJoinDenialReason
+	 */
+	UPROPERTY()
+	uint8 DenialReason = 0;
+};
+
+/** Join in progress data. Holds the current request and any responses. Requests and responses are expected to be cleared in a short amount of time. Combined into one field to reduce field count. */
+USTRUCT()
+struct FPartyMemberJoinInProgressData
+{
+	GENERATED_BODY()
+
+public:
+
+	bool operator==(const FPartyMemberJoinInProgressData& Other) const { return Request == Other.Request && Responses == Other.Responses; }
+	bool operator!=(const FPartyMemberJoinInProgressData& Other) const { return !operator==(Other); }
+
+	/** Current request for the local member. */
+	UPROPERTY()
+	FPartyMemberJoinInProgressRequest Request;
+
+	/** List of responses for other members who requested a reservation. */
+	UPROPERTY()
+	TArray<FPartyMemberJoinInProgressResponse> Responses;
+};
+
 /** Base struct used to replicate data about the state of a single party member to all members. */
 USTRUCT()
 struct PARTY_API FPartyMemberRepData : public FOnlinePartyRepDataBase
@@ -74,6 +139,12 @@ private:
 	UPROPERTY()
 	FString JoinMethod;
 	EXPOSE_REP_DATA_PROPERTY(FPartyMemberRepData, FString, JoinMethod);
+
+	/** Data used for join in progress flow. */
+	UPROPERTY()
+	FPartyMemberJoinInProgressData JoinInProgressData;
+	EXPOSE_USTRUCT_REP_DATA_PROPERTY(FPartyMemberRepData, FPartyMemberJoinInProgressRequest, JoinInProgressData, Request);
+	EXPOSE_USTRUCT_REP_DATA_PROPERTY(FPartyMemberRepData, TArray<FPartyMemberJoinInProgressResponse>, JoinInProgressData, Responses);
 };
 
 using FPartyMemberDataReplicator = TPartyDataReplicator<FPartyMemberRepData, UPartyMember>;
diff --git a/Engine/Plugins/Online/OnlineFramework/Source/Party/Public/Party/SocialParty.h b/Engine/Plugins/Online/OnlineFramework/Source/Party/Public/Party/SocialParty.h
--- a/Engine/Plugins/Online/OnlineFramework/Source/Party/Public/Party/SocialParty.h
+++ b/Engine/Plugins/Online/OnlineFramework/Source/Party/Public/Party/SocialParty.h
@@ -24,6 +24,8 @@ class USocialUser;
 class FOnlineSessionSettings;
 class FOnlineSessionSearchResult;
 enum class EMemberExitedReason : uint8;
+struct FPartyMemberJoinInProgressRequest;
+struct FPartyMemberJoinInProgressResponse;
 
 /** Base struct used to replicate data about the state of the party to all members. */
 USTRUCT()
@@ -195,10 +197,11 @@ public:
 	FOnInviteSent& OnInviteSent() const { return OnInviteSentEvent; }
 
 	DECLARE_EVENT_TwoParams(USocialParty, FOnPartyJIPApproved, const FOnlinePartyId&, bool /* Success*/);
-	UE_DEPRECATED(5.0, "Use OnPartyJIPResponse instead of OnPartyJIPApproved")
+	UE_DEPRECATED(5.1, "Use the new join in progress flow with USocialParty::RequestJoinInProgress.")
 	FOnPartyJIPApproved& OnPartyJIPApproved() const { return OnPartyJIPApprovedEvent; }
 
 	DECLARE_EVENT_ThreeParams(USocialParty, FOnPartyJIPResponse, const FOnlinePartyId&, bool /* Success*/, const FString& /*DeniedResultCode*/);
+	UE_DEPRECATED(5.1, "Use the new join in progress flow with USocialParty::RequestJoinInProgress.")
 	FOnPartyJIPResponse& OnPartyJIPResponse() const { return OnPartyJIPResponseEvent; }
 
 	DECLARE_EVENT_TwoParams(USocialParty, FOnPartyMemberConnectionStatusChanged, UPartyMember&, EMemberConnectionStatus);
@@ -212,6 +215,9 @@ public:
 	virtual void JoinSessionCompleteAnalytics(const FSessionId& SessionId, const FString& JoinBootableGroupSessionResult);
 	bool IsCurrentlyLeaving() const;
 
+	DECLARE_DELEGATE_OneParam(FOnRequestJoinInProgressComplete, const EPartyJoinDenialReason /*DenialReason*/);
+	void RequestJoinInProgress(const UPartyMember& TargetMember, const FOnRequestJoinInProgressComplete& CompletionDelegate);
+
 protected:
 	void InitializeParty(const TSharedRef<const FOnlineParty>& InOssParty);
 	bool IsInitialized() const;
@@ -267,6 +273,7 @@ protected:
 	virtual FPartyJoinApproval EvaluateJoinRequest(const TArray<IOnlinePartyUserPendingJoinRequestInfoConstRef>& Players, bool bFromJoinRequest) const;
 
 	/** Determines the joinability of the game a party is in for JoinInProgress */
+	UE_DEPRECATED(5.1, "Use the new join in progress flow with USocialParty::RequestJoinInProgress.")
 	virtual FPartyJoinApproval EvaluateJIPRequest(const FUniqueNetId& PlayerId) const;
 
 	/** Determines the reason why, if at all, this party is currently flat-out unjoinable  */
@@ -329,7 +336,6 @@ private:
 
 	void HandlePreClientTravel(const FString& PendingURL, ETravelType TravelType, bool bIsSeamlessTravel);
 
-
 	UPartyMember* GetMemberInternal(const FUniqueNetIdRepl& MemberId) const;
 
 private:	// Handlers
@@ -339,11 +345,13 @@ private:	// Handlers
 	void HandlePartyDataReceived(const FUniqueNetId& LocalUserId, const FOnlinePartyId& PartyId, const FName& Namespace, const FOnlinePartyData& PartyData);
 	void HandleJoinabilityQueryReceived(const FUniqueNetId& LocalUserId, const FOnlinePartyId& PartyId, const IOnlinePartyPendingJoinRequestInfo& JoinRequestInfo);
 	void HandlePartyJoinRequestReceived(const FUniqueNetId& LocalUserId, const FOnlinePartyId& PartyId, const IOnlinePartyPendingJoinRequestInfo& JoinRequestInfo);
+	UE_DEPRECATED(5.1, "Use the new join in progress flow with USocialParty::RequestJoinInProgress.")
 	void HandlePartyJIPRequestReceived(const FUniqueNetId& LocalUserId, const FOnlinePartyId& PartyId, const FUniqueNetId& SenderId);
 	void HandlePartyLeft(const FUniqueNetId& LocalUserId, const FOnlinePartyId& PartyId);
 	void HandlePartyMemberExited(const FUniqueNetId& LocalUserId, const FOnlinePartyId& PartyId, const FUniqueNetId& MemberId, EMemberExitedReason ExitReason);
 	void HandlePartyMemberDataReceived(const FUniqueNetId& LocalUserId, const FOnlinePartyId& PartyId, const FUniqueNetId& MemberId, const FName& Namespace, const FOnlinePartyData& PartyMemberData);
 	void HandlePartyMemberJoined(const FUniqueNetId& LocalUserId, const FOnlinePartyId& PartyId, const FUniqueNetId& MemberId);
+	UE_DEPRECATED(5.1, "Use the new join in progress flow with USocialParty::RequestJoinInProgress.")
 	void HandlePartyMemberJIP(const FUniqueNetId& LocalUserId, const FOnlinePartyId& PartyId, bool Success, int32 DeniedResultCode);
 	void HandlePartyMemberPromoted(const FUniqueNetId& LocalUserId, const FOnlinePartyId& PartyId, const FUniqueNetId& NewLeaderId);
 	void HandlePartyPromotionLockoutChanged(const FUniqueNetId& LocalUserId, const FOnlinePartyId& PartyId, bool bArePromotionsLocked);
@@ -359,6 +367,10 @@ private:	// Handlers
 	void HandleRemoveLocalPlayerComplete(const FUniqueNetId& LocalUserId, const FOnlinePartyId& PartyId, ELeavePartyCompletionResult LeaveResult, FOnLeavePartyAttemptComplete OnAttemptComplete);
 
 	void RemovePlayerFromReservationBeacon(const FUniqueNetId& LocalUserId, const FUniqueNetId& PlayerToRemove);
+
+	void HandleJoinInProgressDataRequestChanged(const FPartyMemberJoinInProgressRequest& Request, UPartyMember* Member);
+	void HandleJoinInProgressDataResponsesChanged(const TArray<FPartyMemberJoinInProgressResponse>& Responses, UPartyMember* Member);
+
 private:
 	TSharedPtr<const FOnlineParty> OssParty;
 
@@ -404,6 +416,7 @@ private:
 		FUniqueNetIdRepl RecipientId;
 		TArray<FMemberInfo> Members;
 		bool bIsJIPApproval;
+		int64 JoinInProgressRequestTime = 0;
 		bool bIsPlayerRemoval = false;
 	};
 	TQueue<FPendingMemberApproval> PendingApprovals;
@@ -445,6 +458,27 @@ private:
 	bool bHasReceivedRepData = false;
 	TOptional<bool> bIsRequestingShutdown;
 
+	void RespondToJoinInProgressRequest(const FPendingMemberApproval& PendingApproval, const EPartyJoinDenialReason DenialReason);
+	void CallJoinInProgressComplete(const EPartyJoinDenialReason DenialReason);
+	void RunJoinInProgressTimer();
+
+	/** Complete delegate for join in progress requests. This should only have one at a time. */
+	TOptional<FOnRequestJoinInProgressComplete> RequestJoinInProgressComplete;
+
+	FTimerHandle JoinInProgressTimerHandle;
+
+	/** How often the timer should check in seconds for stale data when running. */
+	UPROPERTY(config)
+	float JoinInProgressTimerRate = 5.f;
+	
+	/** How long in seconds before join in progress requests timeout and are cleared from member data. */
+	UPROPERTY(config)
+	int32 JoinInProgressRequestTimeout = 30;
+
+	/** How long in seconds before join in progress responses are cleared from member data. */
+	UPROPERTY(config)
+	int32 JoinInProgressResponseTimeout = 60;
+
 	mutable FLeavePartyEvent OnPartyLeaveBeginEvent;
 	mutable FLeavePartyEvent OnPartyLeftEvent;
 	mutable FDisconnectPartyEvent OnPartyDisconnectedEvent;
diff --git a/Engine/Plugins/Online/OnlineSubsystem/Source/Public/Interfaces/OnlinePartyInterface.h b/Engine/Plugins/Online/OnlineSubsystem/Source/Public/Interfaces/OnlinePartyInterface.h
--- a/Engine/Plugins/Online/OnlineSubsystem/Source/Public/Interfaces/OnlinePartyInterface.h
+++ b/Engine/Plugins/Online/OnlineSubsystem/Source/Public/Interfaces/OnlinePartyInterface.h
@@ -999,9 +999,10 @@ PARTY_DECLARE_DELEGATETYPE(OnPartyStateChanged);
 * @param DeniedResultCode - descriptive reason for a failure to join
 */
 DECLARE_MULTICAST_DELEGATE_ThreeParams(F_PREFIX(OnPartyJIP), const FUniqueNetId& /*LocalUserId*/, const FOnlinePartyId& /*PartyId*/, bool /*Success*/);
-UE_DEPRECATED(5.1, "Use FOnPartyJIPResponseDelegate instead of FOnPartyJIPDelegate")
+UE_DEPRECATED(5.1, "Use the new join in progress flow with USocialParty::RequestJoinInProgress.")
 typedef FOnPartyJIP::FDelegate FOnPartyJIPDelegate;
 DECLARE_MULTICAST_DELEGATE_FourParams(F_PREFIX(OnPartyJIPResponse), const FUniqueNetId& /*LocalUserId*/, const FOnlinePartyId& /*PartyId*/, bool /*Success*/, int32 /*DeniedResultCode*/);
+UE_DEPRECATED(5.1, "Use the new join in progress flow with USocialParty::RequestJoinInProgress.")
 PARTY_DECLARE_DELEGATETYPE(OnPartyJIPResponse);
 
 /**
@@ -1179,6 +1180,7 @@ PARTY_DECLARE_DELEGATETYPE(OnPartyJoinRequestReceived);
  * @param PartyData - data provided by the sender for the leader to use to determine joinability
  */
 DECLARE_MULTICAST_DELEGATE_ThreeParams(F_PREFIX(OnPartyJIPRequestReceived), const FUniqueNetId& /*LocalUserId*/, const FOnlinePartyId& /*PartyId*/, const FUniqueNetId& /*SenderId*/);
+UE_DEPRECATED(5.1, "Use the new join in progress flow with USocialParty::RequestJoinInProgress.")
 PARTY_DECLARE_DELEGATETYPE(OnPartyJIPRequestReceived);
 
 /**
@@ -1347,7 +1349,8 @@ public:
 	*
 	* @return true if task was started
 	*/
-	virtual bool JIPFromWithinParty(const FUniqueNetId& LocalUserId, const FOnlinePartyId& PartyId, const FUniqueNetId& PartyLeaderId) = 0;
+	UE_DEPRECATED(5.1, "Use the new join in progress flow with USocialParty::RequestJoinInProgress.")
+	virtual bool JIPFromWithinParty(const FUniqueNetId& LocalUserId, const FOnlinePartyId& PartyId, const FUniqueNetId& PartyLeaderId) { return false; }
 
 	/**
 	 * Query a party to check it's current joinability
@@ -1436,7 +1439,8 @@ public:
 	*
 	* @return true if task was started
 	*/
-	virtual bool ApproveJIPRequest(const FUniqueNetId& LocalUserId, const FOnlinePartyId& PartyId, const FUniqueNetId& RecipientId, bool bIsApproved, int32 DeniedResultCode = 0) = 0;
+	UE_DEPRECATED(5.1, "Use the new join in progress flow with USocialParty::RequestJoinInProgress.")
+	virtual bool ApproveJIPRequest(const FUniqueNetId& LocalUserId, const FOnlinePartyId& PartyId, const FUniqueNetId& RecipientId, bool bIsApproved, int32 DeniedResultCode = 0) { return false; }
 
 	/**
 	 * Respond to a query joinability request.  This reflects the current party's joinability state and can change from moment to moment, and therefore does not guarantee a successful join.
@@ -1809,7 +1813,9 @@ public:
 	* @param Success - whether the join in progress action succeeded
 	* @param DeniedResultCode - descriptive reason for a failure to join
 	*/
+	PRAGMA_DISABLE_DEPRECATION_WARNINGS
 	DEFINE_ONLINE_DELEGATE_FOUR_PARAM(OnPartyJIPResponse, const FUniqueNetId& /*LocalUserId*/, const FOnlinePartyId& /*PartyId*/, bool /*Success*/, int32 /*DeniedResultCode*/);
+	PRAGMA_ENABLE_DEPRECATION_WARNINGS
 
 	/**
 	 * Notification when player promotion is locked out.
@@ -1938,7 +1944,9 @@ public:
 	* @param Platform - platform of member that sent the request
 	* @param PartyData - data provided by the sender for the leader to use to determine joinability
 	*/
+	PRAGMA_DISABLE_DEPRECATION_WARNINGS
 	DEFINE_ONLINE_DELEGATE_THREE_PARAM(OnPartyJIPRequestReceived, const FUniqueNetId& /*LocalUserId*/, const FOnlinePartyId& /*PartyId*/, const FUniqueNetId& /*SenderId*/);
+	PRAGMA_ENABLE_DEPRECATION_WARNINGS
 
 	/**
 	 * Notification when a player wants to know if the party is in a joinable state
