Add scope counters around party creation complete to identify cause of a hitch

diff --git a/Engine/Plugins/Online/OnlineFramework/Source/Party/Private/Party/SocialParty.cpp b/Engine/Plugins/Online/OnlineFramework/Source/Party/Private/Party/SocialParty.cpp
--- a/Engine/Plugins/Online/OnlineFramework/Source/Party/Private/Party/SocialParty.cpp
+++ b/Engine/Plugins/Online/OnlineFramework/Source/Party/Private/Party/SocialParty.cpp
@@ -470,6 +470,7 @@ void USocialParty::InitializeParty(const TSharedRef<const FOnlineParty>& InOssPa
 
 void USocialParty::InitializePartyInternal()
 {
+	QUICK_SCOPE_CYCLE_COUNTER(STAT_SocialParty_InitializePartyInternal);
 	IOnlinePartyPtr PartyInterface = Online::GetPartyInterfaceChecked(GetWorld());
 	PartyInterface->AddOnPartyConfigChangedDelegate_Handle(FOnPartyConfigChangedDelegate::CreateUObject(this, &USocialParty::HandlePartyConfigChanged));
 	PartyInterface->AddOnPartyDataReceivedDelegate_Handle(FOnPartyDataReceivedDelegate::CreateUObject(this, &USocialParty::HandlePartyDataReceived));
@@ -1486,6 +1487,7 @@ TSubclassOf<UPartyMember> USocialParty::GetDesiredMemberClass(bool bLocalPlayer)
 
 void USocialParty::HandlePartyStateChanged(const FUniqueNetId& LocalUserId, const FOnlinePartyId& PartyId, EPartyState PartyState, EPartyState PreviousPartyState)
 {
+	QUICK_SCOPE_CYCLE_COUNTER(STAT_SocialParty_HandlePartyStateChanged);
 	if (PartyState == EPartyState::Disconnected)
 	{
 		// If we transition to the disconnected state, then we are lacking an XMPP connection (or logged out of MCP?)
diff --git a/Engine/Plugins/Online/OnlineFramework/Source/Party/Private/SocialManager.cpp b/Engine/Plugins/Online/OnlineFramework/Source/Party/Private/SocialManager.cpp
--- a/Engine/Plugins/Online/OnlineFramework/Source/Party/Private/SocialManager.cpp
+++ b/Engine/Plugins/Online/OnlineFramework/Source/Party/Private/SocialManager.cpp
@@ -882,6 +882,7 @@ const USocialManager::FJoinPartyAttempt* USocialManager::GetJoinAttemptInProgres
 
 USocialParty* USocialManager::EstablishNewParty(const FUniqueNetId& LocalUserId, const FOnlinePartyId& PartyId, const FOnlinePartyTypeId& PartyTypeId)
 {
+	QUICK_SCOPE_CYCLE_COUNTER(STAT_SocialManager_EstablishNewParty);
 	IOnlinePartyPtr PartyInterface = Online::GetPartyInterfaceChecked(GetWorld());
 	
 	TSubclassOf<USocialParty> PartyClass = GetPartyClassForType(PartyTypeId);
@@ -896,7 +897,10 @@ USocialParty* USocialManager::EstablishNewParty(const FUniqueNetId& LocalUserId,
 		// This must be done before InitializeParty(), as initialization can complete synchronously.
 		JoinedPartiesByTypeId.Add(PartyTypeId, NewParty);
 
-		NewParty->InitializeParty(OssParty.ToSharedRef());
+		{
+			QUICK_SCOPE_CYCLE_COUNTER(STAT_SocialManager_EstablishNewParty_InitializeParty);
+			NewParty->InitializeParty(OssParty.ToSharedRef());
+		}
 
 		if (NewParty->IsPersistentParty())
 		{
@@ -1125,6 +1129,7 @@ void USocialManager::HandleCreatePartyComplete(const FUniqueNetId& LocalUserId,
 {
 	ABORT_DURING_SHUTDOWN();
 
+	QUICK_SCOPE_CYCLE_COUNTER(STAT_SocialManager_HandleCreatePartyComplete);
 	ECreatePartyCompletionResult LocalCreationResult = Result;
 	if (Result == ECreatePartyCompletionResult::Succeeded)
 	{
@@ -1145,7 +1150,10 @@ void USocialManager::HandleCreatePartyComplete(const FUniqueNetId& LocalUserId,
 	}
 
 	UE_LOG(LogParty, Verbose, TEXT("Finished trying to create party [%s] with result [%s]"), PartyId.IsValid() ? *PartyId->ToDebugString() : TEXT("Invalid"), ToString(LocalCreationResult));
-	CompletionDelegate.ExecuteIfBound(LocalCreationResult);
+	{
+		QUICK_SCOPE_CYCLE_COUNTER(STAT_SocialManager_HandleCreatePartyComplete_Delegate);
+		CompletionDelegate.ExecuteIfBound(LocalCreationResult);
+	}
 }
 
 void USocialManager::HandleJoinPartyComplete(const FUniqueNetId& LocalUserId, const FOnlinePartyId& PartyId, EJoinPartyCompletionResult Result, int32 NotApprovedReasonCode, FOnlinePartyTypeId PartyTypeId)
