Add additional join in progress timestamp for clearing party data to handle clock skew between clients.

diff --git a/Engine/Plugins/Online/OnlineFramework/Source/Party/Private/Party/SocialParty.cpp b/Engine/Plugins/Online/OnlineFramework/Source/Party/Private/Party/SocialParty.cpp
--- a/Engine/Plugins/Online/OnlineFramework/Source/Party/Private/Party/SocialParty.cpp
+++ b/Engine/Plugins/Online/OnlineFramework/Source/Party/Private/Party/SocialParty.cpp
@@ -1999,11 +1999,12 @@ void USocialParty::RespondToJoinInProgressRequest(const FPendingMemberApproval&
 
 	FPartyMemberJoinInProgressResponse Response;
 	Response.Requester = RequestingMember->GetPrimaryNetId();
-	Response.Time = PendingApproval.JoinInProgressRequestTime;
+	Response.RequestTime = PendingApproval.JoinInProgressRequestTime;
+	Response.ResponseTime = FDateTime::UtcNow().ToUnixTimestamp();
 	Response.DenialReason = static_cast<uint8>(DenialReason);
 
-	UE_LOG(LogParty, Verbose, TEXT("RespondToJoinInProgressRequest: Sending response Requester=%s Time=%d DenialReason=%s"),
-		*Response.Requester.ToDebugString(), Response.Time, ToString(DenialReason));
+	UE_LOG(LogParty, Verbose, TEXT("RespondToJoinInProgressRequest: Sending response Requester=%s RequestTime=%d ResponseTime=%d DenialReason=%s"),
+		*Response.Requester.ToDebugString(), Response.RequestTime, Response.ResponseTime, ToString(DenialReason));
 	TArray<FPartyMemberJoinInProgressResponse> Responses = GetOwningLocalMember().GetRepData().GetJoinInProgressDataResponses();
 	Responses.Add(Response);
 	GetOwningLocalMember().GetMutableRepData().SetJoinInProgressDataResponses(Responses);
@@ -2022,7 +2023,7 @@ void USocialParty::HandleJoinInProgressDataResponsesChanged(const TArray<FPartyM
 
 	for (const FPartyMemberJoinInProgressResponse& Response : Responses)
 	{
-		if (Response.Time != Request.Time)
+		if (Response.RequestTime != Request.Time)
 		{
 			// Response was not for us.
 			continue;
@@ -2033,8 +2034,8 @@ void USocialParty::HandleJoinInProgressDataResponsesChanged(const TArray<FPartyM
 
 		if (!RequestingMember)
 		{
-			UE_LOG(LogParty, Warning, TEXT("HandleJoinInProgressDataResponsesChanged: Could not find member for response Requester=%s Time=%d DenialReason=%s"),
-				*Response.Requester.ToDebugString(), Response.Time, ToString(DenialReason));
+			UE_LOG(LogParty, Warning, TEXT("HandleJoinInProgressDataResponsesChanged: Could not find member for response Requester=%s RequestTime=%d ResponseTime=%d DenialReason=%s"),
+				*Response.Requester.ToDebugString(), Response.RequestTime, Response.ResponseTime, ToString(DenialReason));
 			continue;
 		}
 
@@ -2046,8 +2047,8 @@ void USocialParty::HandleJoinInProgressDataResponsesChanged(const TArray<FPartyM
 
 		// Responses are ordered newest first, so use the first one we find.
 		const int64 Now = FDateTime::UtcNow().ToUnixTimestamp();
-		UE_LOG(LogParty, Verbose, TEXT("HandleJoinInProgressDataResponsesChanged: Received response Requester=%s Time=%d DenialReason=%s ResponseTime=%d"),
-			*Response.Requester.ToDebugString(), Response.Time, ToString(DenialReason), Now - Response.Time);
+		UE_LOG(LogParty, Verbose, TEXT("HandleJoinInProgressDataResponsesChanged: Received response Requester=%s RequestTime=%d ResponseTime=%d DenialReason=%s RTT=%d"),
+			*Response.Requester.ToDebugString(), Response.RequestTime, Response.ResponseTime, ToString(DenialReason), Now - Response.RequestTime);
 		CallJoinInProgressComplete(DenialReason);
 		break;
 	}
@@ -2092,7 +2093,7 @@ void USocialParty::RunJoinInProgressTimer()
 	TArray<FPartyMemberJoinInProgressResponse> ResponsesToKeep;
 	for (const FPartyMemberJoinInProgressResponse& Response : Responses)
 	{
-		const int64 Expires = Response.Time + JoinInProgressResponseTimeout;
+		const int64 Expires = Response.ResponseTime + JoinInProgressResponseTimeout;
 		if (Expires > Now)
 		{
 			ResponsesToKeep.Add(Response);
diff --git a/Engine/Plugins/Online/OnlineFramework/Source/Party/Public/Party/PartyMember.h b/Engine/Plugins/Online/OnlineFramework/Source/Party/Public/Party/PartyMember.h
--- a/Engine/Plugins/Online/OnlineFramework/Source/Party/Public/Party/PartyMember.h
+++ b/Engine/Plugins/Online/OnlineFramework/Source/Party/Public/Party/PartyMember.h
@@ -65,7 +65,13 @@ struct FPartyMemberJoinInProgressResponse
 	GENERATED_BODY()
 
 public:
-	bool operator==(const FPartyMemberJoinInProgressResponse& Other) const { return Requester == Other.Requester && Time == Other.Time; }
+	bool operator==(const FPartyMemberJoinInProgressResponse& Other) const
+	{
+		return Requester == Other.Requester &&
+			RequestTime == Other.RequestTime &&
+			ResponseTime == Other.ResponseTime &&
+			DenialReason == Other.DenialReason;
+	}
 	bool operator!=(const FPartyMemberJoinInProgressResponse& Other) const { return !operator==(Other); }
 
 	/** Remote member that this response is for. */
@@ -74,8 +80,12 @@ public:
 
 	/** Time the request was made. Matches FPartyMemberJoinInProgressRequest::Time */
 	UPROPERTY()
-	int64 Time = 0;
-	
+	int64 RequestTime = 0;
+
+	/** Time the response was made. */
+	UPROPERTY()
+	int64 ResponseTime = 0;
+
 	/**
 	 * Result of session reservation attempt.
 	 * @see EPartyJoinDenialReason
