- Fix ABORT_DURING_SHUTDOWN not aborting during shutdown

diff --git a/Engine/Plugins/Online/OnlineFramework/Source/Party/Public/SocialManager.h b/Engine/Plugins/Online/OnlineFramework/Source/Party/Public/SocialManager.h
--- a/Engine/Plugins/Online/OnlineFramework/Source/Party/Public/SocialManager.h
+++ b/Engine/Plugins/Online/OnlineFramework/Source/Party/Public/SocialManager.h
@@ -25,7 +25,7 @@ class USocialDebugTools;
 
 enum ETravelType;
 
-#define ABORT_DURING_SHUTDOWN()  if (bShutdownPending) { UE_LOG(LogParty, Log, TEXT("%s - Received callback after bShutdownPending."), ANSI_TO_TCHAR(__FUNCTION__)); return; }
+#define ABORT_DURING_SHUTDOWN() if (IsEngineExitRequested() || bShutdownPending) { UE_LOG(LogParty, Log, TEXT("%s - Received callback during shutdown: IsEngineExitRequested=%s, bShutdownPending=%s."), ANSI_TO_TCHAR(__FUNCTION__), *LexToString(IsEngineExitRequested()), *LexToString(bShutdownPending)); return; }
 
 /** Singleton manager at the top of the social framework */
 UCLASS(Within = GameInstance, Config = Game)
