Remove OnlineTracing dependency from OnlineFramework.

diff --git a/Engine/Plugins/Online/OnlineFramework/OnlineFramework.uplugin b/Engine/Plugins/Online/OnlineFramework/OnlineFramework.uplugin
--- a/Engine/Plugins/Online/OnlineFramework/OnlineFramework.uplugin
+++ b/Engine/Plugins/Online/OnlineFramework/OnlineFramework.uplugin
@@ -58,11 +58,6 @@
 		{
 			"Name": "OnlineSubsystemUtils",
 			"Enabled": true
-		},
-		{
-			"Name": "OnlineTracing",
-			"Enabled": true,
-			"Optional": true
 		}
 	]
 }
\ No newline at end of file
diff --git a/Engine/Plugins/Online/OnlineFramework/Source/Hotfix/Hotfix.Build.cs b/Engine/Plugins/Online/OnlineFramework/Source/Hotfix/Hotfix.Build.cs
--- a/Engine/Plugins/Online/OnlineFramework/Source/Hotfix/Hotfix.Build.cs
+++ b/Engine/Plugins/Online/OnlineFramework/Source/Hotfix/Hotfix.Build.cs
@@ -28,13 +28,6 @@ public class Hotfix : ModuleRules
 			}
 			);
 
-		bool bHasOnlineTracing = Directory.Exists(Path.Combine(EngineDirectory, "Restricted", "NotForLicensees", "Plugins", "Online", "OnlineTracing"));
-		if (bHasOnlineTracing)
-		{
-			PublicDefinitions.Add("WITH_ONLINETRACING=1");
-			PrivateDependencyModuleNames.Add("OnlineTracing");
-		}
-
 		PublicDefinitions.Add("UPDATEMANAGER_PLATFORM_ENVIRONMENT_DETECTION=" + (bPlatformEnvironmentDetection ? "1" : "0"));
 	}
 
diff --git a/Engine/Plugins/Online/OnlineFramework/Source/Hotfix/Private/OnlineHotfixManager.cpp b/Engine/Plugins/Online/OnlineFramework/Source/Hotfix/Private/OnlineHotfixManager.cpp
--- a/Engine/Plugins/Online/OnlineFramework/Source/Hotfix/Private/OnlineHotfixManager.cpp
+++ b/Engine/Plugins/Online/OnlineFramework/Source/Hotfix/Private/OnlineHotfixManager.cpp
@@ -31,10 +31,6 @@
 
 #include UE_INLINE_GENERATED_CPP_BY_NAME(OnlineHotfixManager)
 
-#ifdef WITH_ONLINETRACING
-#include "OnlineTracingModule.h"
-#endif
-
 DEFINE_LOG_CATEGORY(LogHotfixManager);
 
 /** This character must be between important pieces of file information (platform, initype, version) */
@@ -922,7 +918,6 @@ bool UOnlineHotfixManager::HotfixIniFile(const FString& FileName, const FString&
 	bool bUpdateLogSuppression = false;
 	bool bUpdateConsoleVariables = false;
 	bool bUpdateHttpConfigs = false;
-	bool bUpdateOnlineTracing = false;
 	TSet<FString> OnlineSubSections;
 	TSet<FString> UpdatedSectionNames;
 	// Find the set of object classes that were affected
@@ -984,7 +979,6 @@ bool UOnlineHotfixManager::HotfixIniFile(const FString& FileName, const FString&
 					const TCHAR* ConsoleVariableSection = TEXT("[ConsoleVariables]");
 					const TCHAR* HttpSection = TEXT("[HTTP"); // note "]" omitted on purpose since we want a partial match
 					const TCHAR* OnlineSubSectionKey = TEXT("[OnlineSubsystem"); // note "]" omitted on purpose since we want a partial match
-					const TCHAR* OnlineTracingSection = TEXT("[OnlineTracing]");
 					if (!bUpdateLogSuppression && FCString::Strnicmp(*IniData + StartIndex, LogConfigSection, FCString::Strlen(LogConfigSection)) == 0)
 					{
 						bUpdateLogSuppression = true;
@@ -997,10 +991,6 @@ bool UOnlineHotfixManager::HotfixIniFile(const FString& FileName, const FString&
 					{
 						bUpdateHttpConfigs = true;
 					}
-					else if (!bUpdateOnlineTracing && FCString::Strnicmp(*IniData + StartIndex, OnlineTracingSection, FCString::Strlen(OnlineTracingSection)) == 0)
-					{
-						bUpdateOnlineTracing = true;
-					}
 					else if (FCString::Strnicmp(*IniData + StartIndex, OnlineSubSectionKey, FCString::Strlen(OnlineSubSectionKey)) == 0)
 					{
 						FString SectionStr = IniData.Mid(StartIndex, EndIndex - StartIndex + 1);
@@ -1127,15 +1117,6 @@ PRAGMA_ENABLE_DEPRECATION_WARNINGS
 		FHttpModule::Get().UpdateConfigs();
 	}
 
-	// Reload configs for tracing system, this may init or tear it down if enable is toggled
-#ifdef WITH_ONLINETRACING
-	if (bUpdateOnlineTracing && 
-		FOnlineTracingModule::IsAvailable())
-	{
-		FOnlineTracingModule::Get().UpdateConfig();
-	}
-#endif
-
 	// Reload configs relevant to OSS config sections that were updated
 	IOnlineSubsystem* OnlineSub = IOnlineSubsystem::Get(OSSName.Len() ? FName(*OSSName, FNAME_Find) : NAME_None);
 	if (OnlineSub != nullptr)
