Specify encryption data for all beacon connections.

diff --git a/Engine/Plugins/Online/OnlineFramework/Source/Party/Private/Party/SocialParty.cpp b/Engine/Plugins/Online/OnlineFramework/Source/Party/Private/Party/SocialParty.cpp
--- a/Engine/Plugins/Online/OnlineFramework/Source/Party/Private/Party/SocialParty.cpp
+++ b/Engine/Plugins/Online/OnlineFramework/Source/Party/Private/Party/SocialParty.cpp
@@ -1487,6 +1487,11 @@ TSubclassOf<UPartyMember> USocialParty::GetDesiredMemberClass(bool bLocalPlayer)
 	return UPartyMember::StaticClass();
 }
 
+bool USocialParty::InitializeBeaconEncryptionData(AOnlineBeaconClient& BeaconClient, const FString& SessionId)
+{
+	return true;
+}
+
 void USocialParty::HandlePartyStateChanged(const FUniqueNetId& LocalUserId, const FOnlinePartyId& PartyId, EPartyState PartyState, EPartyState PreviousPartyState)
 {
 	QUICK_SCOPE_CYCLE_COUNTER(STAT_SocialParty_HandlePartyStateChanged);
@@ -1536,29 +1541,32 @@ void USocialParty::ConnectToReservationBeacon()
 
 							UE_LOG(LogParty, Verbose, TEXT("Party [%s] created reservation beacon [%s]."), *ToDebugString(), *LocalReservationBeaconClient->GetName());
 
-							LocalReservationBeaconClient->OnHostConnectionFailure().BindUObject(this, &USocialParty::HandleBeaconHostConnectionFailed);
-							LocalReservationBeaconClient->OnReservationRequestComplete().BindUObject(this, &USocialParty::HandleReservationRequestComplete);
-
-							TArray<FPlayerReservation> ReservationAsArray;
-							ReservationAsArray.Reserve(NextApproval.Members.Num());
-							for (const FPendingMemberApproval::FMemberInfo& MemberInfo : NextApproval.Members)
+							if (InitializeBeaconEncryptionData(*ReservationBeaconClient, Session->GetSessionIdStr()))
 							{
-								FPlayerReservation& Reservation = ReservationAsArray.Emplace_GetRef();
-								Reservation.UniqueId = MemberInfo.MemberId;
-								Reservation.Platform = MemberInfo.Platform;
+								LocalReservationBeaconClient->OnHostConnectionFailure().BindUObject(this, &USocialParty::HandleBeaconHostConnectionFailed);
+								LocalReservationBeaconClient->OnReservationRequestComplete().BindUObject(this, &USocialParty::HandleReservationRequestComplete);
 
-								if (!NextApproval.bIsJIPApproval && MemberInfo.JoinData.IsValid())
+								TArray<FPlayerReservation> ReservationAsArray;
+								ReservationAsArray.Reserve(NextApproval.Members.Num());
+								for (const FPendingMemberApproval::FMemberInfo& MemberInfo : NextApproval.Members)
 								{
-									const ECrossplayPreference CrossplayPreference = GetCrossplayPreferenceFromJoinData(*MemberInfo.JoinData);
-									Reservation.bAllowCrossplay = (CrossplayPreference == ECrossplayPreference::OptedIn);
+									FPlayerReservation& Reservation = ReservationAsArray.Emplace_GetRef();
+									Reservation.UniqueId = MemberInfo.MemberId;
+									Reservation.Platform = MemberInfo.Platform;
+
+									if (!NextApproval.bIsJIPApproval && MemberInfo.JoinData.IsValid())
+									{
+										const ECrossplayPreference CrossplayPreference = GetCrossplayPreferenceFromJoinData(*MemberInfo.JoinData);
+										Reservation.bAllowCrossplay = (CrossplayPreference == ECrossplayPreference::OptedIn);
+									}
+									else
+									{
+										Reservation.bAllowCrossplay = true; // This will not matter since we are JIP, and the session already has crossplay set.
+									}
 								}
-								else
-								{
-									Reservation.bAllowCrossplay = true; // This will not matter since we are JIP, and the session already has crossplay set.
-								}
-							}
 
-							bStartedConnection = ReservationBeaconClient->RequestReservationUpdate(URL, Session->GetSessionIdStr(), GetPartyLeader()->GetPrimaryNetId(), ReservationAsArray, NextApproval.bIsPlayerRemoval);
+								bStartedConnection = ReservationBeaconClient->RequestReservationUpdate(URL, Session->GetSessionIdStr(), GetPartyLeader()->GetPrimaryNetId(), ReservationAsArray, NextApproval.bIsPlayerRemoval);
+							}
 						}
 					}
 				}
@@ -1629,7 +1637,7 @@ ASpectatorBeaconClient* USocialParty::CreateSpectatorBeaconClient()
 	LastSpectatorBeaconClientNetDriverName = NAME_None;
 	SpectatorBeaconClient = World->SpawnActor<ASpectatorBeaconClient>(SpectatorBeaconClientClass);
 
-	return SpectatorBeaconClient;
+	return SpectatorBeaconClient.Get();
 }
 
 void USocialParty::PumpApprovalQueue()
@@ -1735,14 +1743,14 @@ void USocialParty::CleanupReservationBeacon()
 
 void USocialParty::CleanupSpectatorBeacon()
 {
-	if (SpectatorBeaconClient)
+	if (ASpectatorBeaconClient* LocalSpectatorBeaconClient = SpectatorBeaconClient.Get())
 	{
 		UE_LOG(LogParty, Verbose, TEXT("Spectator reservation beacon cleanup while in state %s, pending approvals: %s"), ToString(SpectatorBeaconClient->GetConnectionState()), !PendingApprovals.IsEmpty() ? TEXT("true") : TEXT("false"));
 
-		LastReservationBeaconClientNetDriverName = SpectatorBeaconClient->GetNetDriverName();
-		SpectatorBeaconClient->OnHostConnectionFailure().Unbind();
-		SpectatorBeaconClient->OnReservationRequestComplete().Unbind();
-		SpectatorBeaconClient->DestroyBeacon();
+		LastReservationBeaconClientNetDriverName = LocalSpectatorBeaconClient->GetNetDriverName();
+		LocalSpectatorBeaconClient->OnHostConnectionFailure().Unbind();
+		LocalSpectatorBeaconClient->OnReservationRequestComplete().Unbind();
+		LocalSpectatorBeaconClient->DestroyBeacon();
 		SpectatorBeaconClient = nullptr;
 	}
 }
diff --git a/Engine/Plugins/Online/OnlineFramework/Source/Party/Public/Party/SocialParty.h b/Engine/Plugins/Online/OnlineFramework/Source/Party/Public/Party/SocialParty.h
--- a/Engine/Plugins/Online/OnlineFramework/Source/Party/Public/Party/SocialParty.h
+++ b/Engine/Plugins/Online/OnlineFramework/Source/Party/Public/Party/SocialParty.h
@@ -282,6 +282,9 @@ protected:
 	/** Override in child classes to specify the type of UPartyMember to create */
 	virtual TSubclassOf<UPartyMember> GetDesiredMemberClass(bool bLocalPlayer) const;
 
+	/** Override in child classes to provide encryption data for party beacon connections. */
+	virtual bool InitializeBeaconEncryptionData(AOnlineBeaconClient& BeaconClient, const FString& SessionId);
+
 	bool IsInviteRateLimited(const USocialUser& User, ESocialSubsystem SubsystemType) const;
 
 	bool ApplyCrossplayRestriction(FPartyJoinApproval& JoinApproval, const FUserPlatform& Platform, const FOnlinePartyData& JoinData) const;
@@ -305,7 +308,7 @@ protected:
 	void CleanupSpectatorBeacon();
 	ASpectatorBeaconClient* CreateSpectatorBeaconClient();
 
-	ASpectatorBeaconClient* GetSpectatorBeaconClient() const { return SpectatorBeaconClient; }
+	ASpectatorBeaconClient* GetSpectatorBeaconClient() const { return SpectatorBeaconClient.Get(); }
 
 	/** Child classes MUST call EstablishRepDataInstance() on this using their member rep data struct instance */
 	FPartyDataReplicator PartyDataReplicator;
@@ -444,7 +447,7 @@ private:
 	
 	/** Spectator beacon client instance while getting approval for spectator*/
 	UPROPERTY()
-	TObjectPtr<ASpectatorBeaconClient> SpectatorBeaconClient = nullptr;
+	TWeakObjectPtr<ASpectatorBeaconClient> SpectatorBeaconClient = nullptr;
 
 	/**
 	 * True when we have limited functionality due to lacking an xmpp connection.
diff --git a/Engine/Plugins/Online/OnlineSubsystemUtils/Source/OnlineSubsystemUtils/Private/OnlineBeaconClient.cpp b/Engine/Plugins/Online/OnlineSubsystemUtils/Source/OnlineSubsystemUtils/Private/OnlineBeaconClient.cpp
--- a/Engine/Plugins/Online/OnlineSubsystemUtils/Source/OnlineSubsystemUtils/Private/OnlineBeaconClient.cpp
+++ b/Engine/Plugins/Online/OnlineSubsystemUtils/Source/OnlineSubsystemUtils/Private/OnlineBeaconClient.cpp
@@ -352,16 +352,27 @@ void AOnlineBeaconClient::NotifyControlMessage(UNetConnection* Connection, uint8
 		{
 		case NMT_EncryptionAck:
 			{
-				if (FNetDelegates::OnReceivedNetworkEncryptionAck.IsBound())
+				// Enable encryption using the beacons encryption data when set.
+				if (!EncryptionData.Identifier.IsEmpty())
 				{
-					TWeakObjectPtr<UNetConnection> WeakConnection = Connection;
-					FNetDelegates::OnReceivedNetworkEncryptionAck.Execute(FOnEncryptionKeyResponse::CreateUObject(this, &ThisClass::FinalizeEncryptedConnection, WeakConnection));
+					FEncryptionKeyResponse Response;
+					Response.Response = EEncryptionResponse::Success;
+					Response.EncryptionData = EncryptionData;
+					FinalizeEncryptedConnection(Response, Connection);
 				}
 				else
 				{
-					// Force close the session
-					UE_LOG(LogBeacon, Warning, TEXT("%s: No delegate available to handle encryption ack, disconnecting."), *Connection->GetName());
-					OnFailure();
+					if (FNetDelegates::OnReceivedNetworkEncryptionAck.IsBound())
+					{
+						TWeakObjectPtr<UNetConnection> WeakConnection = Connection;
+						FNetDelegates::OnReceivedNetworkEncryptionAck.Execute(FOnEncryptionKeyResponse::CreateUObject(this, &ThisClass::FinalizeEncryptedConnection, WeakConnection));
+					}
+					else
+					{
+						// Force close the session
+						UE_LOG(LogBeacon, Warning, TEXT("%s: No delegate available to handle encryption ack, disconnecting."), *Connection->GetName());
+						OnFailure();
+					}
 				}
 				break;
 			}
