Prevent ignored patch check from running on dedicated server.

diff --git a/Engine/Plugins/Online/OnlineFramework/Source/PatchCheck/Private/PatchCheck.cpp b/Engine/Plugins/Online/OnlineFramework/Source/PatchCheck/Private/PatchCheck.cpp
--- a/Engine/Plugins/Online/OnlineFramework/Source/PatchCheck/Private/PatchCheck.cpp
+++ b/Engine/Plugins/Online/OnlineFramework/Source/PatchCheck/Private/PatchCheck.cpp
@@ -146,28 +146,38 @@ void FPatchCheck::StartPlatformOSSPatchCheck()
 
 void FPatchCheck::StartOSSPatchCheck()
 {
-	EPatchCheckResult PatchResult = EPatchCheckResult::PatchCheckFailure;
-	bool bStarted = false;
+	if (SkipPatchCheck())
+	{
+		UE_LOG(LogPatchCheck, Verbose, TEXT("[StartOSSPatchCheck] Skipping patch check."));
 
-	// Online::GetIdentityInterface() can take a UWorld for correctness, but that only matters in PIE right now
-	// and update checks should never happen in PIE currently.
-	IOnlineIdentityPtr IdentityInt = Online::GetIdentityInterface();
-	if (IdentityInt.IsValid())
+		// Trigger completion if check is skipped.
+		PatchCheckComplete(EPatchCheckResult::NoPatchRequired);
+	}
+	else
 	{
-		// User could be invalid for "before title/login" check, underlying code doesn't need a valid user currently
-		FUniqueNetIdPtr UserId = IdentityInt->CreateUniquePlayerId(TEXT("InvalidUser"));
-		if (UserId.IsValid())
+		EPatchCheckResult PatchResult = EPatchCheckResult::PatchCheckFailure;
+		bool bStarted = false;
+
+		// Online::GetIdentityInterface() can take a UWorld for correctness, but that only matters in PIE right now
+		// and update checks should never happen in PIE currently.
+		IOnlineIdentityPtr IdentityInt = Online::GetIdentityInterface();
+		if (IdentityInt.IsValid())
 		{
-			bStarted = true;
-			IdentityInt->GetUserPrivilege(*UserId,
-				EUserPrivileges::CanPlayOnline, IOnlineIdentity::FOnGetUserPrivilegeCompleteDelegate::CreateRaw(this, &FPatchCheck::OnCheckForPatchComplete, false));
+			// User could be invalid for "before title/login" check, underlying code doesn't need a valid user currently
+			FUniqueNetIdPtr UserId = IdentityInt->CreateUniquePlayerId(TEXT("InvalidUser"));
+			if (UserId.IsValid())
+			{
+				bStarted = true;
+				IdentityInt->GetUserPrivilege(*UserId,
+					EUserPrivileges::CanPlayOnline, IOnlineIdentity::FOnGetUserPrivilegeCompleteDelegate::CreateRaw(this, &FPatchCheck::OnCheckForPatchComplete, false));
+			}
 		}
-	}
 
-	if (!bStarted)
-	{
-		// Any failure to call GetUserPrivilege will result in completing the flow via this path
-		PatchCheckComplete(PatchResult);
+		if (!bStarted)
+		{
+			// Any failure to call GetUserPrivilege will result in completing the flow via this path
+			PatchCheckComplete(PatchResult);
+		}
 	}
 }
 
@@ -218,31 +228,28 @@ void FPatchCheck::OnCheckForPatchComplete(const FUniqueNetId& UniqueId, EUserPri
 	EPatchCheckResult Result = EPatchCheckResult::NoPatchRequired;
 	if (Privilege == EUserPrivileges::CanPlayOnline)
 	{
-		if (bConsoleCheck || !SkipPatchCheck())
+		if (PrivilegeResult & (uint32)IOnlineIdentity::EPrivilegeResults::RequiredSystemUpdate)
 		{
-			if (PrivilegeResult & (uint32)IOnlineIdentity::EPrivilegeResults::RequiredSystemUpdate)
-			{
-				Result = EPatchCheckResult::PatchRequired;
-			}
-			else if (PrivilegeResult & (uint32)IOnlineIdentity::EPrivilegeResults::RequiredPatchAvailable)
-			{
-				Result = EPatchCheckResult::PatchRequired;
-			}
-			else if (PrivilegeResult & ((uint32)IOnlineIdentity::EPrivilegeResults::UserNotLoggedIn | (uint32)IOnlineIdentity::EPrivilegeResults::UserNotFound))
+			Result = EPatchCheckResult::PatchRequired;
+		}
+		else if (PrivilegeResult & (uint32)IOnlineIdentity::EPrivilegeResults::RequiredPatchAvailable)
+		{
+			Result = EPatchCheckResult::PatchRequired;
+		}
+		else if (PrivilegeResult & ((uint32)IOnlineIdentity::EPrivilegeResults::UserNotLoggedIn | (uint32)IOnlineIdentity::EPrivilegeResults::UserNotFound))
+		{
+			Result = EPatchCheckResult::NoLoggedInUser;
+		}
+		else if (PrivilegeResult & (uint32)IOnlineIdentity::EPrivilegeResults::GenericFailure)
+		{
+			CA_CONSTANT_IF(PATCH_CHECK_FAIL_ON_GENERIC_FAILURE)
 			{
-				Result = EPatchCheckResult::NoLoggedInUser;
+				Result = EPatchCheckResult::PatchCheckFailure;
 			}
-			else if (PrivilegeResult & (uint32)IOnlineIdentity::EPrivilegeResults::GenericFailure)
+			else
 			{
-				CA_CONSTANT_IF(PATCH_CHECK_FAIL_ON_GENERIC_FAILURE)
-				{
-					Result = EPatchCheckResult::PatchCheckFailure;
-				}
-				else
-				{
-					// Skip console backend failures
-					Result = EPatchCheckResult::NoPatchRequired;
-				}
+				// Skip console backend failures
+				Result = EPatchCheckResult::NoPatchRequired;
 			}
 		}
 	}
