Add ability for types of social users to override which interactions are available for a given user
#tests verified when overriden in child class, overriden function is used

diff --git a/Engine/Plugins/Online/OnlineFramework/Source/Party/Public/User/SocialUser.h b/Engine/Plugins/Online/OnlineFramework/Source/Party/Public/User/SocialUser.h
--- a/Engine/Plugins/Online/OnlineFramework/Source/Party/Public/User/SocialUser.h
+++ b/Engine/Plugins/Online/OnlineFramework/Source/Party/Public/User/SocialUser.h
@@ -94,8 +94,8 @@ public:
 	bool SetUserLocalAttribute(ESocialSubsystem SubsystemType, const FString& AttrName, const FString& AttrValue);
 	bool GetUserAttribute(ESocialSubsystem SubsystemType, const FString& AttrName, FString& OutAttrValue) const;
 
-	bool HasAnyInteractionsAvailable() const;
-	TArray<FSocialInteractionHandle> GetAllAvailableInteractions() const;
+	virtual bool HasAnyInteractionsAvailable() const;
+	virtual TArray<FSocialInteractionHandle> GetAllAvailableInteractions() const;
 
 	virtual bool CanSendFriendInvite(ESocialSubsystem SubsystemType) const;
 	virtual bool SendFriendInvite(ESocialSubsystem SubsystemType);
