Add support for hotfixing app specific dedicated server files.

diff --git a/Engine/Plugins/Online/OnlineFramework/Source/Hotfix/Private/OnlineHotfixManager.cpp b/Engine/Plugins/Online/OnlineFramework/Source/Hotfix/Private/OnlineHotfixManager.cpp
--- a/Engine/Plugins/Online/OnlineFramework/Source/Hotfix/Private/OnlineHotfixManager.cpp
+++ b/Engine/Plugins/Online/OnlineFramework/Source/Hotfix/Private/OnlineHotfixManager.cpp
@@ -223,7 +223,7 @@ void UOnlineHotfixManager::PostInitProperties()
 	PlatformPrefix = DebugPrefix + ANSI_TO_TCHAR(FPlatformProperties::PlatformName());
 	PlatformPrefix += HOTFIX_SEPARATOR;
 	// Server prefix
-	ServerPrefix = DebugPrefix + TEXT("DedicatedServer");
+	ServerPrefix = DebugPrefix + GetDedicatedServerPrefix();
 	// Build the default prefix too
 	DefaultPrefix = DebugPrefix + TEXT("Default");
 
@@ -1937,6 +1937,11 @@ bool UOnlineHotfixManager::ShouldPerformHotfix()
 	return IsRunningGame() || IsRunningDedicatedServer() || IsRunningClientOnly();
 }
 
+FString UOnlineHotfixManager::GetDedicatedServerPrefix() const
+{
+	return TEXT("DedicatedServer");
+}
+
 UWorld* UOnlineHotfixManager::GetWorld() const
 {
 	return OwnerWorld.IsValid() ? OwnerWorld.Get() : nullptr;
diff --git a/Engine/Plugins/Online/OnlineFramework/Source/Hotfix/Public/OnlineHotfixManager.h b/Engine/Plugins/Online/OnlineFramework/Source/Hotfix/Public/OnlineHotfixManager.h
--- a/Engine/Plugins/Online/OnlineFramework/Source/Hotfix/Public/OnlineHotfixManager.h
+++ b/Engine/Plugins/Online/OnlineFramework/Source/Hotfix/Public/OnlineHotfixManager.h
@@ -331,6 +331,9 @@ protected:
 
 	virtual bool ShouldPerformHotfix();
 
+	/** Allow the application to override the dedicated server filename prefix. */
+	virtual FString GetDedicatedServerPrefix() const;
+
 public:
 	UOnlineHotfixManager();
 	UOnlineHotfixManager(FVTableHelper& Helper);
