Add MPA implementation config, use token instead of json for join info when the config is enabled. Because the json string is too long(more than the limit of 256) for connection string in MPA.

diff --git a/Engine/Plugins/Online/OnlineFramework/Source/Party/Private/Party/PartyPlatformSessionMonitor.cpp b/Engine/Plugins/Online/OnlineFramework/Source/Party/Private/Party/PartyPlatformSessionMonitor.cpp
--- a/Engine/Plugins/Online/OnlineFramework/Source/Party/Private/Party/PartyPlatformSessionMonitor.cpp
+++ b/Engine/Plugins/Online/OnlineFramework/Source/Party/Private/Party/PartyPlatformSessionMonitor.cpp
@@ -1003,6 +1003,10 @@ void FPartyPlatformSessionMonitor::HandleDestroySessionComplete(FName SessionNam
 	ShutdownInternal();
 }
 
+#if PARTY_PLATFORM_SESSIONS_XBL
+extern TAutoConsoleVariable<bool> CVarXboxMpaEnabled;
+#endif
+
 bool FPartyPlatformSessionMonitor::ConfigurePlatformSessionSettings(FOnlineSessionSettings& SessionSettings) const
 {
 	bool bEstablishedPartySettings = false;
@@ -1010,27 +1014,40 @@ bool FPartyPlatformSessionMonitor::ConfigurePlatformSessionSettings(FOnlineSessi
 	IOnlinePartyPtr PartyInterface = MonitoredParty.IsValid() ? Online::GetPartyInterface(MonitoredParty->GetWorld()) : nullptr;
 	if (PartyInterface.IsValid())
 	{
-		const FString JoinInfoJson = PartyInterface->MakeJoinInfoJson(*MonitoredParty->GetOwningLocalUserId(), MonitoredParty->GetPartyId());
-		if (ensure(!JoinInfoJson.IsEmpty()))
+#if PARTY_PLATFORM_SESSIONS_XBL
+		if (CVarXboxMpaEnabled.GetValueOnAnyThread())
+		{
+			if (IOnlinePartyJoinInfoConstPtr JoinInfo = PartyInterface->MakeJoinInfo(*MonitoredParty->GetOwningLocalUserId(), MonitoredParty->GetPartyId()))
+			{
+				SessionSettings.Set(SETTING_CUSTOM_JOIN_INFO, PartyInterface->MakeTokenFromJoinInfo(*JoinInfo), EOnlineDataAdvertisementType::ViaOnlineService);
+				bEstablishedPartySettings = true;
+			}
+		}
+		else
+#endif
 		{
-			bEstablishedPartySettings = true;
+			const FString JoinInfoJson = PartyInterface->MakeJoinInfoJson(*MonitoredParty->GetOwningLocalUserId(), MonitoredParty->GetPartyId());
+			if (ensure(!JoinInfoJson.IsEmpty()))
+			{
+				bEstablishedPartySettings = true;
 
 #if PARTY_PLATFORM_SESSIONS_PSN
-			SessionSettings.Set(SETTING_HOST_MIGRATION, true, EOnlineDataAdvertisementType::DontAdvertise);
-			SessionSettings.Set(SETTING_CUSTOM, JoinInfoJson, EOnlineDataAdvertisementType::ViaOnlineService);
+				SessionSettings.Set(SETTING_HOST_MIGRATION, true, EOnlineDataAdvertisementType::DontAdvertise);
+				SessionSettings.Set(SETTING_CUSTOM, JoinInfoJson, EOnlineDataAdvertisementType::ViaOnlineService);
 #elif PARTY_PLATFORM_SESSIONS_XBL
-			// This needs to match our value on the XDP service configuration
-			SessionSettings.Set(SETTING_SESSION_TEMPLATE_NAME, FString(TEXT("MultiplayerGameSession")), EOnlineDataAdvertisementType::DontAdvertise);
-
-			// XBOX has their own value for this as SETTING_CUSTOM is hard-coded to constant data in the OSS, and is the actual originator of
-			// SETTING_CUSTOM.  Everyone else co-opted it and made it dynamic, so we need to use something else just here so other OSS' still
-			// work out of the box for this functionality
-			// Encode our JoinInfo into Base64 to prevent XboxLive from parsing our json
-			SessionSettings.Set(SETTING_CUSTOM_JOIN_INFO, FBase64::Encode(JoinInfoJson), EOnlineDataAdvertisementType::ViaOnlineService);
+				// This needs to match our value on the XDP service configuration
+				SessionSettings.Set(SETTING_SESSION_TEMPLATE_NAME, FString(TEXT("MultiplayerGameSession")), EOnlineDataAdvertisementType::DontAdvertise);
+
+				// XBOX has their own value for this as SETTING_CUSTOM is hard-coded to constant data in the OSS, and is the actual originator of
+				// SETTING_CUSTOM.  Everyone else co-opted it and made it dynamic, so we need to use something else just here so other OSS' still
+				// work out of the box for this functionality
+				// Encode our JoinInfo into Base64 to prevent XboxLive from parsing our json
+				SessionSettings.Set(SETTING_CUSTOM_JOIN_INFO, FBase64::Encode(JoinInfoJson), EOnlineDataAdvertisementType::ViaOnlineService);
 #elif PLATFORM_DESKTOP
-			// PC (Tencent)
-			SessionSettings.Set(SETTING_CUSTOM, JoinInfoJson, EOnlineDataAdvertisementType::ViaOnlineService);
+				// PC (Tencent)
+				SessionSettings.Set(SETTING_CUSTOM, JoinInfoJson, EOnlineDataAdvertisementType::ViaOnlineService);
 #endif
+			}
 		}
 	}
 
diff --git a/Engine/Plugins/Online/OnlineFramework/Source/Party/Private/SocialManager.cpp b/Engine/Plugins/Online/OnlineFramework/Source/Party/Private/SocialManager.cpp
--- a/Engine/Plugins/Online/OnlineFramework/Source/Party/Private/SocialManager.cpp
+++ b/Engine/Plugins/Online/OnlineFramework/Source/Party/Private/SocialManager.cpp
@@ -950,19 +950,28 @@ USocialParty* USocialManager::GetPartyInternal(const FOnlinePartyId& PartyId, bo
 	return nullptr;
 }
 
+#if PARTY_PLATFORM_SESSIONS_XBL
+extern TAutoConsoleVariable<bool> CVarXboxMpaEnabled;
+#endif
+
 TSharedPtr<const IOnlinePartyJoinInfo> USocialManager::GetJoinInfoFromSession(const FOnlineSessionSearchResult& PlatformSession)
 {
 	static const FName JoinInfoSettingName = PARTY_PLATFORM_SESSIONS_XBL ? SETTING_CUSTOM_JOIN_INFO : SETTING_CUSTOM;
 
-	FString JoinInfoJson;
-	if (PlatformSession.Session.SessionSettings.Get(JoinInfoSettingName, JoinInfoJson))
+	FString JoinInfoString;
+	if (PlatformSession.Session.SessionSettings.Get(JoinInfoSettingName, JoinInfoString))
 	{
+		IOnlinePartyPtr PartyInterface = Online::GetPartyInterfaceChecked(GetWorld());
 #if PARTY_PLATFORM_SESSIONS_XBL 
-		// On Xbox we encode our party data in base64 to avoid XboxLive trying to parse our json, so now we need to decode that
-		FBase64::Decode(JoinInfoJson, JoinInfoJson);
+		if (CVarXboxMpaEnabled.GetValueOnAnyThread())
+		{
+			return PartyInterface->MakeJoinInfoFromToken(JoinInfoString);
+		}
+
+		// On Xbox MPSD we encode our party data in base64 to avoid XboxLive trying to parse our json, so now we need to decode that
+		FBase64::Decode(JoinInfoString, JoinInfoString);
 #endif 
-		IOnlinePartyPtr PartyInterface = Online::GetPartyInterfaceChecked(GetWorld());
-		return PartyInterface->MakeJoinInfoFromJson(JoinInfoJson);
+		return PartyInterface->MakeJoinInfoFromJson(JoinInfoString);
 	}
 	return nullptr;
 }
diff --git a/Engine/Plugins/Online/OnlineSubsystem/Source/Public/Interfaces/OnlinePartyInterface.h b/Engine/Plugins/Online/OnlineSubsystem/Source/Public/Interfaces/OnlinePartyInterface.h
--- a/Engine/Plugins/Online/OnlineSubsystem/Source/Public/Interfaces/OnlinePartyInterface.h
+++ b/Engine/Plugins/Online/OnlineSubsystem/Source/Public/Interfaces/OnlinePartyInterface.h
@@ -1740,6 +1740,15 @@ public:
 	 */
 	virtual FString MakeTokenFromJoinInfo(const IOnlinePartyJoinInfo& JoinInfo) const = 0;
 
+	/**
+	 * returns the join info for current party
+	 *
+	 * @param LocalUserId       - user making the request
+	 * @param PartyId           - party to make the json from
+	 *
+ 	 */
+	virtual IOnlinePartyJoinInfoConstPtr MakeJoinInfo(const FUniqueNetId& LocalUserId, const FOnlinePartyId& PartyId) = 0;
+
 	/**
 	 * Creates a IOnlinePartyJoinInfo object from a command line token
 	 *
