Update recently players by members in the party.

diff --git a/Engine/Plugins/Online/OnlineFramework/Source/Party/Private/Party/PartyPlatformSessionMonitor.cpp b/Engine/Plugins/Online/OnlineFramework/Source/Party/Private/Party/PartyPlatformSessionMonitor.cpp
--- a/Engine/Plugins/Online/OnlineFramework/Source/Party/Private/Party/PartyPlatformSessionMonitor.cpp
+++ b/Engine/Plugins/Online/OnlineFramework/Source/Party/Private/Party/PartyPlatformSessionMonitor.cpp
@@ -13,6 +13,7 @@
 #include "HAL/IConsoleManager.h"
 #include "Misc/Base64.h"
 
+#include "Interfaces/OnlineFriendsInterface.h"
 #include "OnlineSubsystemSessionSettings.h"
 #include "OnlineSessionSettings.h"
 #include "OnlineSubsystemUtils.h"
@@ -194,6 +195,11 @@ IOnlineSessionPtr FPartyPlatformSessionManager::GetSessionInterface()
 	return Online::GetSessionInterfaceChecked(SocialManager.GetWorld(), PlatformOssName);
 }
 
+IOnlineFriendsPtr FPartyPlatformSessionManager::GetFriendsInterface()
+{
+	return Online::GetFriendsInterfaceChecked(SocialManager.GetWorld(), PlatformOssName);
+}
+
 FUniqueNetIdRepl FPartyPlatformSessionManager::GetLocalUserPlatformId() const
 {
 	return SocialManager.GetFirstLocalUserId(ESocialSubsystem::Platform);
@@ -504,6 +510,8 @@ void FPartyPlatformSessionMonitor::Initialize()
 	Party.OnPartyMemberLeft().AddSP(this, &FPartyPlatformSessionMonitor::HandlePartyMemberLeft);
 
 	EvaluateCurrentSession();
+
+	UpdateRecentPlayersOfLocalMembers(MonitoredParty->GetPartyMembers());
 }
 
 void FPartyPlatformSessionMonitor::ShutdownInternal()
@@ -813,6 +821,49 @@ void FPartyPlatformSessionMonitor::HandlePartyMemberCreated(UPartyMember& NewMem
 	}
 }
 
+bool FPartyPlatformSessionMonitor::ShouldRecordAsRecentPlayer(const FUniqueNetId& LocalUserId, const UPartyMember* PartyMember)
+{
+	if (PartyMember->GetPlatformOssName() != LocalUserId.GetType())
+	{
+		return false;
+	}
+
+	if (*PartyMember->GetRepData().GetPlatformDataUniqueId() == LocalUserId)
+	{
+		return false;
+	}
+
+	return true;
+}
+
+void FPartyPlatformSessionMonitor::UpdateRecentPlayersOfLocalMembers(const TArray<UPartyMember*>& RecentPlayers)
+{
+	for (UPartyMember* Member : MonitoredParty->GetPartyMembers())
+	{
+		if (Member->IsLocalPlayer())
+		{
+			UpdateRecentPlayersOfLocalUser(*Member->GetRepData().GetPlatformDataUniqueId(), RecentPlayers);
+		}
+	}
+}
+
+void FPartyPlatformSessionMonitor::UpdateRecentPlayersOfLocalUser(const FUniqueNetId& LocalUserId, const TArray<UPartyMember*>& Members)
+{
+	TArray<FReportPlayedWithUser> RecentPlayers;
+	for (UPartyMember* Member : Members)
+	{
+		if (ShouldRecordAsRecentPlayer(LocalUserId, Member))
+		{
+			RecentPlayers.Emplace(Member->GetRepData().GetPlatformDataUniqueId().GetUniqueNetId().ToSharedRef(), "", ERecentPlayerEncounterType::Teammate);
+		}
+	}
+
+	if (!RecentPlayers.IsEmpty())
+	{
+		SessionManager->GetFriendsInterface()->AddRecentPlayers(LocalUserId, RecentPlayers, TEXT(""), FOnAddRecentPlayersComplete());
+	}
+}
+
 void FPartyPlatformSessionMonitor::HandlePartyMemberInitialized(UPartyMember* InitializedMember)
 {
 	if (IsTencentPlatform() && InitializedMember->GetPlatformOssName() == TENCENT_SUBSYSTEM)
@@ -825,6 +876,14 @@ void FPartyPlatformSessionMonitor::HandlePartyMemberInitialized(UPartyMember* In
 	{
 		AddLocalPlayerToSession(InitializedMember);
 	}
+
+	TArray<UPartyMember*> RecentPlayers{ InitializedMember };
+	UpdateRecentPlayersOfLocalMembers(RecentPlayers);
+
+	if (InitializedMember->IsLocalPlayer())
+	{
+		UpdateRecentPlayersOfLocalUser(*InitializedMember->GetRepData().GetPlatformDataUniqueId(), MonitoredParty->GetPartyMembers());
+	}
 }
 
 void FPartyPlatformSessionMonitor::HandlePartyMemberLeft(UPartyMember* OldMember, const EMemberExitedReason ExitReason)
diff --git a/Engine/Plugins/Online/OnlineFramework/Source/Party/Private/Party/PartyPlatformSessionMonitor.h b/Engine/Plugins/Online/OnlineFramework/Source/Party/Private/Party/PartyPlatformSessionMonitor.h
--- a/Engine/Plugins/Online/OnlineFramework/Source/Party/Private/Party/PartyPlatformSessionMonitor.h
+++ b/Engine/Plugins/Online/OnlineFramework/Source/Party/Private/Party/PartyPlatformSessionMonitor.h
@@ -29,6 +29,7 @@ public:
 	bool FindSession(const FPartyPlatformSessionInfo& SessionInfo, const FOnFindSessionAttemptComplete& OnAttemptComplete);
 
 	IOnlineSessionPtr GetSessionInterface();
+	IOnlineFriendsPtr GetFriendsInterface();
 	FUniqueNetIdRepl GetLocalUserPlatformId() const;
 
 private:
@@ -93,6 +94,10 @@ private:
 
 	void ProcessJoinFailure();
 
+	bool ShouldRecordAsRecentPlayer(const FUniqueNetId& LocalUserId, const UPartyMember* PartyMember);
+	void UpdateRecentPlayersOfLocalMembers(const TArray<UPartyMember*>& RecentPlayers);
+	void UpdateRecentPlayersOfLocalUser(const FUniqueNetId& LocalUserId, const TArray<UPartyMember*>& Members);
+
 private:
 	void HandlePlatformSessionsChanged();
 	void HandlePartyConfigurationChanged(const FPartyConfiguration& NewConfig);
