Call create session instead of find and join session when MPA enabled, because there is no find functionality in MPA.

diff --git a/Engine/Plugins/Online/OnlineFramework/Source/Party/Private/Party/PartyPlatformSessionMonitor.cpp b/Engine/Plugins/Online/OnlineFramework/Source/Party/Private/Party/PartyPlatformSessionMonitor.cpp
--- a/Engine/Plugins/Online/OnlineFramework/Source/Party/Private/Party/PartyPlatformSessionMonitor.cpp
+++ b/Engine/Plugins/Online/OnlineFramework/Source/Party/Private/Party/PartyPlatformSessionMonitor.cpp
@@ -104,6 +104,10 @@ static FAutoConsoleVariableRef CVar_EstablishSessionRetryDelay(
 	ECVF_Default
 );
 
+#if PARTY_PLATFORM_SESSIONS_XBL
+extern TAutoConsoleVariable<bool> CVarXboxMpaEnabled;
+#endif
+
 //////////////////////////////////////////////////////////////////////////
 // FPartySessionManager
 //////////////////////////////////////////////////////////////////////////
@@ -475,15 +479,22 @@ void FPartyPlatformSessionMonitor::EvaluateCurrentSession()
 		}
 		else
 		{
-			for (UPartyMember* Member : MonitoredParty->GetPartyMembers())
+			if (ShouldAlwaysCreateLocalPlatformSession())
 			{
-				if (Member->IsLocalPlayer())
+				CreateSession(SessionManager->GetLocalUserPlatformId());
+			}
+			else
+			{
+				for (UPartyMember* Member : MonitoredParty->GetPartyMembers())
 				{
-					if (ExistingSessionInfo->IsSessionOwner(*Member))
+					if (Member->IsLocalPlayer())
 					{
-						// There is no session ID yet, but the session owner is a local player, so it's on us to create it now
-						CreateSession(Member->GetRepData().GetPlatformDataUniqueId());
-						break;
+						if (ExistingSessionInfo->IsSessionOwner(*Member))
+						{
+							// There is no session ID yet, but the session owner is a local player, so it's on us to create it now
+							CreateSession(Member->GetRepData().GetPlatformDataUniqueId());
+							break;
+						}
 					}
 				}
 			}
@@ -821,6 +832,18 @@ void FPartyPlatformSessionMonitor::HandlePartyMemberCreated(UPartyMember& NewMem
 	}
 }
 
+bool FPartyPlatformSessionMonitor::ShouldAlwaysCreateLocalPlatformSession() const
+{
+#if PARTY_PLATFORM_SESSIONS_XBL
+	if (CVarXboxMpaEnabled.GetValueOnAnyThread())
+	{
+		return true;
+	}
+#endif
+
+	return false;
+}
+
 bool FPartyPlatformSessionMonitor::ShouldRecordAsRecentPlayer(const FUniqueNetId& LocalUserId, const UPartyMember* PartyMember)
 {
 	if (PartyMember->GetPlatformOssName() != LocalUserId.GetType())
@@ -1062,10 +1085,6 @@ void FPartyPlatformSessionMonitor::HandleDestroySessionComplete(FName SessionNam
 	ShutdownInternal();
 }
 
-#if PARTY_PLATFORM_SESSIONS_XBL
-extern TAutoConsoleVariable<bool> CVarXboxMpaEnabled;
-#endif
-
 bool FPartyPlatformSessionMonitor::ConfigurePlatformSessionSettings(FOnlineSessionSettings& SessionSettings) const
 {
 	bool bEstablishedPartySettings = false;
diff --git a/Engine/Plugins/Online/OnlineFramework/Source/Party/Private/Party/PartyPlatformSessionMonitor.h b/Engine/Plugins/Online/OnlineFramework/Source/Party/Private/Party/PartyPlatformSessionMonitor.h
--- a/Engine/Plugins/Online/OnlineFramework/Source/Party/Private/Party/PartyPlatformSessionMonitor.h
+++ b/Engine/Plugins/Online/OnlineFramework/Source/Party/Private/Party/PartyPlatformSessionMonitor.h
@@ -94,6 +94,7 @@ private:
 
 	void ProcessJoinFailure();
 
+	bool ShouldAlwaysCreateLocalPlatformSession() const;
 	bool ShouldRecordAsRecentPlayer(const FUniqueNetId& LocalUserId, const UPartyMember* PartyMember);
 	void UpdateRecentPlayersOfLocalMembers(const TArray<UPartyMember*>& RecentPlayers);
 	void UpdateRecentPlayersOfLocalUser(const FUniqueNetId& LocalUserId, const TArray<UPartyMember*>& Members);
