Improve Logging for USocialParty UpdatePartyConfig Completion

#tests verified that it single file compiled successfully.

diff --git a/Engine/Plugins/Online/OnlineFramework/Source/Party/Private/Party/SocialParty.cpp b/Engine/Plugins/Online/OnlineFramework/Source/Party/Private/Party/SocialParty.cpp
--- a/Engine/Plugins/Online/OnlineFramework/Source/Party/Private/Party/SocialParty.cpp
+++ b/Engine/Plugins/Online/OnlineFramework/Source/Party/Private/Party/SocialParty.cpp
@@ -928,16 +928,33 @@ void USocialParty::HandlePartyConfigChanged(const FUniqueNetId& LocalUserId, con
 
 void USocialParty::HandleUpdatePartyConfigComplete(const FUniqueNetId& LocalUserId, const FOnlinePartyId& PartyId, EUpdateConfigCompletionResult Result)
 {
-	if (Result == EUpdateConfigCompletionResult::Succeeded)
+	if (!OssParty.IsValid())
 	{
-		UE_LOG(LogParty, Verbose, TEXT("[%s] Party config updated %s"), *PartyId.ToDebugString(), ToString(Result));
+		UE_LOG(LogParty, Error, TEXT("[%s] Invalid OssParty value. Party %s"), ANSI_TO_TCHAR(__FUNCTION__), *PartyId.ToDebugString());
+		return;
+	}
 
+	if (Result == EUpdateConfigCompletionResult::Succeeded)
+	{
 		CurrentConfig = *OssParty->GetConfiguration();
+		UE_LOG(LogParty, Verbose, TEXT("[%s] Success updating Party Config for %s - OssParty Configuration:%s"),
+			ANSI_TO_TCHAR(__FUNCTION__),
+			*PartyId.ToDebugString(),
+			*::ToDebugString(CurrentConfig)
+		);
+
 		OnPartyConfigurationChanged().Broadcast(CurrentConfig);
 	}
 	else
-	{
-		UE_LOG(LogParty, Warning, TEXT("Failed to update config for party [%s]"), *PartyId.ToDebugString());
+	{	
+		const FPartyConfiguration& OssConfig = *OssParty->GetConfiguration();
+		UE_LOG(LogParty, Error, TEXT("[%s] Failed to update Party Config for %s - Current Configuration:%s Oss Configuration:%s Result:%s"),
+			ANSI_TO_TCHAR(__FUNCTION__),
+			*PartyId.ToDebugString(),
+			*::ToDebugString(CurrentConfig),
+			*::ToDebugString(OssConfig),
+			ToString(Result)
+		);
 	}
 }
 
