Fix "Add Playstation Network Friend" button not appearing for PS5 users.

#rnx
#Tests Tested on a Preflight build by QA, verified fixed.

diff --git a/Engine/Plugins/Online/OnlineFramework/Source/Party/Private/Interactions/CoreInteractions.cpp b/Engine/Plugins/Online/OnlineFramework/Source/Party/Private/Interactions/CoreInteractions.cpp
--- a/Engine/Plugins/Online/OnlineFramework/Source/Party/Private/Interactions/CoreInteractions.cpp
+++ b/Engine/Plugins/Online/OnlineFramework/Source/Party/Private/Interactions/CoreInteractions.cpp
@@ -54,7 +54,7 @@ PRAGMA_DISABLE_DEPRECATION_WARNINGS
 		return LOCTEXT("AddPlatformFriend_Live", "Add Xbox Live Friend");
 	}
 PRAGMA_ENABLE_DEPRECATION_WARNINGS
-	else if (PlatformOssName == PS4_SUBSYSTEM)
+	else if (PlatformOssName == PS4_SUBSYSTEM || USocialSettings::IsSonyOSS(PlatformOssName))
 	{
 		return LOCTEXT("AddPlatformFriend_PSN", "Add Playstation Network Friend");
 	}
diff --git a/Engine/Plugins/Online/OnlineFramework/Source/Party/Private/SocialSettings.cpp b/Engine/Plugins/Online/OnlineFramework/Source/Party/Private/SocialSettings.cpp
--- a/Engine/Plugins/Online/OnlineFramework/Source/Party/Private/SocialSettings.cpp
+++ b/Engine/Plugins/Online/OnlineFramework/Source/Party/Private/SocialSettings.cpp
@@ -109,3 +109,16 @@ const FSocialPlatformDescription* USocialSettings::GetSocialPlatformDescriptionF
 	});
 }
 
+// Static
+bool USocialSettings::IsSonyOSS(const FName& OSSName)
+{
+	const USocialSettings& SettingsCDO = *GetDefault<USocialSettings>();
+	for (const FName& SonyOSS : SettingsCDO.SonyOSSNames)
+	{
+		if (SonyOSS == OSSName)
+		{
+			return true;
+		}
+	}
+	return false;
+}
\ No newline at end of file
diff --git a/Engine/Plugins/Online/OnlineFramework/Source/Party/Private/User/SocialUser.cpp b/Engine/Plugins/Online/OnlineFramework/Source/Party/Private/User/SocialUser.cpp
--- a/Engine/Plugins/Online/OnlineFramework/Source/Party/Private/User/SocialUser.cpp
+++ b/Engine/Plugins/Online/OnlineFramework/Source/Party/Private/User/SocialUser.cpp
@@ -894,7 +894,7 @@ bool USocialUser::CanSendFriendInvite(ESocialSubsystem SubsystemType) const
 		//@todo DanH: Really need OssCaps or something to be able to just ask an OSS if it supports a given feature. For now, we just magically know that we only support sending XB, PSN, and WeGame invites
 		const FName PlatformOssName = USocialManager::GetSocialOssName(ESocialSubsystem::Platform);
 PRAGMA_DISABLE_DEPRECATION_WARNINGS
-		if (PlatformOssName != LIVE_SUBSYSTEM && PlatformOssName != PS4_SUBSYSTEM && PlatformOssName != TENCENT_SUBSYSTEM)
+		if (PlatformOssName != LIVE_SUBSYSTEM && PlatformOssName != PS4_SUBSYSTEM && !USocialSettings::IsSonyOSS(PlatformOssName) && PlatformOssName != TENCENT_SUBSYSTEM)
 		{
 			return false;
 		}
diff --git a/Engine/Plugins/Online/OnlineFramework/Source/Party/Public/SocialSettings.h b/Engine/Plugins/Online/OnlineFramework/Source/Party/Public/SocialSettings.h
--- a/Engine/Plugins/Online/OnlineFramework/Source/Party/Public/SocialSettings.h
+++ b/Engine/Plugins/Online/OnlineFramework/Source/Party/Public/SocialSettings.h
@@ -65,6 +65,10 @@ public:
 	static int32 GetMinNicknameLength();
 	static int32 GetMaxNicknameLength();
 	static const TArray<FSocialPlatformDescription>& GetSocialPlatformDescriptions();
+
+	/** Checks if current OnlineSubsystem name is Sony specific */
+	static bool IsSonyOSS(const FName& OSSName);
+
 	/**
 	 * Get a platform description (from GetSocialPlatformDescriptions) for a specific OnlineSubsystem.
 	 * @param OnlineSubsystemName the online subsystem name to search for
@@ -116,4 +120,7 @@ private:
 
 	UPROPERTY(Config)
 	TArray<FSocialPlatformDescription> SocialPlatformDescriptions;
+
+	UPROPERTY(Config)
+	TArray<FName> SonyOSSNames;
 };
\ No newline at end of file
