Fix the live crash due to empty platform id when recording recent players.

diff --git a/Engine/Plugins/Online/OnlineFramework/Source/Party/Private/Party/PartyPlatformSessionMonitor.cpp b/Engine/Plugins/Online/OnlineFramework/Source/Party/Private/Party/PartyPlatformSessionMonitor.cpp
--- a/Engine/Plugins/Online/OnlineFramework/Source/Party/Private/Party/PartyPlatformSessionMonitor.cpp
+++ b/Engine/Plugins/Online/OnlineFramework/Source/Party/Private/Party/PartyPlatformSessionMonitor.cpp
@@ -522,7 +522,12 @@ void FPartyPlatformSessionMonitor::Initialize()
 
 	EvaluateCurrentSession();
 
-	UpdateRecentPlayersOfLocalMembers(MonitoredParty->GetPartyMembers());
+#if PARTY_PLATFORM_SESSIONS_XBL
+	if (CVarXboxMpaEnabled.GetValueOnAnyThread())
+	{
+		UpdateRecentPlayersOfLocalMembers(MonitoredParty->GetPartyMembers());
+	}
+#endif
 }
 
 void FPartyPlatformSessionMonitor::ShutdownInternal()
@@ -851,6 +856,11 @@ bool FPartyPlatformSessionMonitor::ShouldRecordAsRecentPlayer(const FUniqueNetId
 		return false;
 	}
 
+	if (!PartyMember->GetRepData().GetPlatformDataUniqueId())
+	{
+		return false;
+	}
+
 	if (*PartyMember->GetRepData().GetPlatformDataUniqueId() == LocalUserId)
 	{
 		return false;
@@ -865,7 +875,10 @@ void FPartyPlatformSessionMonitor::UpdateRecentPlayersOfLocalMembers(const TArra
 	{
 		if (Member->IsLocalPlayer())
 		{
-			UpdateRecentPlayersOfLocalUser(*Member->GetRepData().GetPlatformDataUniqueId(), RecentPlayers);
+			if (Member->GetRepData().GetPlatformDataUniqueId())
+			{
+				UpdateRecentPlayersOfLocalUser(*Member->GetRepData().GetPlatformDataUniqueId(), RecentPlayers);
+			}
 		}
 	}
 }
@@ -900,19 +913,34 @@ void FPartyPlatformSessionMonitor::HandlePartyMemberInitialized(UPartyMember* In
 		AddLocalPlayerToSession(InitializedMember);
 	}
 
-	TArray<UPartyMember*> RecentPlayers{ InitializedMember };
-	UpdateRecentPlayersOfLocalMembers(RecentPlayers);
-
-	if (InitializedMember->IsLocalPlayer())
+#if PARTY_PLATFORM_SESSIONS_XBL
+	if (CVarXboxMpaEnabled.GetValueOnAnyThread())
 	{
-		UpdateRecentPlayersOfLocalUser(*InitializedMember->GetRepData().GetPlatformDataUniqueId(), MonitoredParty->GetPartyMembers());
+		TArray<UPartyMember*> RecentPlayers{ InitializedMember };
+		UpdateRecentPlayersOfLocalMembers(RecentPlayers);
+
+		if (InitializedMember->IsLocalPlayer())
+		{
+			if (InitializedMember->GetRepData().GetPlatformDataUniqueId())
+			{
+				UpdateRecentPlayersOfLocalUser(*InitializedMember->GetRepData().GetPlatformDataUniqueId(), MonitoredParty->GetPartyMembers());
+			}
+		}	
 	}
+#endif
 }
 
 void FPartyPlatformSessionMonitor::HandlePartyMemberLeft(UPartyMember* OldMember, const EMemberExitedReason ExitReason)
 {
-	UE_LOG(LogParty, Verbose, TEXT("HandlePartyMemberLeft: PartyMember=%s User=%s ExitReason=%s"),
-		*OldMember->ToDebugString(true), *OldMember->GetRepData().GetPlatformDataUniqueId().ToDebugString(), ToString(ExitReason));
+	if (OldMember && OldMember->GetRepData().GetPlatformDataUniqueId())
+	{
+		UE_LOG(LogParty, Verbose, TEXT("HandlePartyMemberLeft: PartyMember=%s User=%s ExitReason=%s"),
+			*OldMember->ToDebugString(true), *OldMember->GetRepData().GetPlatformDataUniqueId().ToDebugString(), ToString(ExitReason));
+	}
+	else
+	{
+		UE_LOG(LogParty, Warn, TEXT("HandlePartyMemberLeft: PartyMember=Unknown ExitReason=%s"), ToString(ExitReason));
+	}
 
 	if (IsTencentPlatform() && OldMember->GetPlatformOssName() == TENCENT_SUBSYSTEM)
 	{
