[Backout] - CL23765291  - CIS Compile Editor Error
Original CL Desc
-----------------------------------------------------------------
Fix the live crash due to empty platform id when recording recent players.

diff --git a/Engine/Plugins/Online/OnlineFramework/Source/Party/Private/Party/PartyPlatformSessionMonitor.cpp b/Engine/Plugins/Online/OnlineFramework/Source/Party/Private/Party/PartyPlatformSessionMonitor.cpp
--- a/Engine/Plugins/Online/OnlineFramework/Source/Party/Private/Party/PartyPlatformSessionMonitor.cpp
+++ b/Engine/Plugins/Online/OnlineFramework/Source/Party/Private/Party/PartyPlatformSessionMonitor.cpp
@@ -522,12 +522,7 @@ void FPartyPlatformSessionMonitor::Initialize()
 
 	EvaluateCurrentSession();
 
-#if PARTY_PLATFORM_SESSIONS_XBL
-	if (CVarXboxMpaEnabled.GetValueOnAnyThread())
-	{
-		UpdateRecentPlayersOfLocalMembers(MonitoredParty->GetPartyMembers());
-	}
-#endif
+	UpdateRecentPlayersOfLocalMembers(MonitoredParty->GetPartyMembers());
 }
 
 void FPartyPlatformSessionMonitor::ShutdownInternal()
@@ -856,11 +851,6 @@ bool FPartyPlatformSessionMonitor::ShouldRecordAsRecentPlayer(const FUniqueNetId
 		return false;
 	}
 
-	if (!PartyMember->GetRepData().GetPlatformDataUniqueId())
-	{
-		return false;
-	}
-
 	if (*PartyMember->GetRepData().GetPlatformDataUniqueId() == LocalUserId)
 	{
 		return false;
@@ -875,10 +865,7 @@ void FPartyPlatformSessionMonitor::UpdateRecentPlayersOfLocalMembers(const TArra
 	{
 		if (Member->IsLocalPlayer())
 		{
-			if (Member->GetRepData().GetPlatformDataUniqueId())
-			{
-				UpdateRecentPlayersOfLocalUser(*Member->GetRepData().GetPlatformDataUniqueId(), RecentPlayers);
-			}
+			UpdateRecentPlayersOfLocalUser(*Member->GetRepData().GetPlatformDataUniqueId(), RecentPlayers);
 		}
 	}
 }
@@ -913,34 +900,19 @@ void FPartyPlatformSessionMonitor::HandlePartyMemberInitialized(UPartyMember* In
 		AddLocalPlayerToSession(InitializedMember);
 	}
 
-#if PARTY_PLATFORM_SESSIONS_XBL
-	if (CVarXboxMpaEnabled.GetValueOnAnyThread())
-	{
-		TArray<UPartyMember*> RecentPlayers{ InitializedMember };
-		UpdateRecentPlayersOfLocalMembers(RecentPlayers);
+	TArray<UPartyMember*> RecentPlayers{ InitializedMember };
+	UpdateRecentPlayersOfLocalMembers(RecentPlayers);
 
-		if (InitializedMember->IsLocalPlayer())
-		{
-			if (InitializedMember->GetRepData().GetPlatformDataUniqueId())
-			{
-				UpdateRecentPlayersOfLocalUser(*InitializedMember->GetRepData().GetPlatformDataUniqueId(), MonitoredParty->GetPartyMembers());
-			}
-		}	
+	if (InitializedMember->IsLocalPlayer())
+	{
+		UpdateRecentPlayersOfLocalUser(*InitializedMember->GetRepData().GetPlatformDataUniqueId(), MonitoredParty->GetPartyMembers());
 	}
-#endif
 }
 
 void FPartyPlatformSessionMonitor::HandlePartyMemberLeft(UPartyMember* OldMember, const EMemberExitedReason ExitReason)
 {
-	if (OldMember && OldMember->GetRepData().GetPlatformDataUniqueId())
-	{
-		UE_LOG(LogParty, Verbose, TEXT("HandlePartyMemberLeft: PartyMember=%s User=%s ExitReason=%s"),
-			*OldMember->ToDebugString(true), *OldMember->GetRepData().GetPlatformDataUniqueId().ToDebugString(), ToString(ExitReason));
-	}
-	else
-	{
-		UE_LOG(LogParty, Warn, TEXT("HandlePartyMemberLeft: PartyMember=Unknown ExitReason=%s"), ToString(ExitReason));
-	}
+	UE_LOG(LogParty, Verbose, TEXT("HandlePartyMemberLeft: PartyMember=%s User=%s ExitReason=%s"),
+		*OldMember->ToDebugString(true), *OldMember->GetRepData().GetPlatformDataUniqueId().ToDebugString(), ToString(ExitReason));
 
 	if (IsTencentPlatform() && OldMember->GetPlatformOssName() == TENCENT_SUBSYSTEM)
 	{
