Fix the live crash due to empty platform id when recording recent players.

diff --git a/Engine/Plugins/Online/OnlineFramework/Source/Party/Private/Party/PartyPlatformSessionMonitor.cpp b/Engine/Plugins/Online/OnlineFramework/Source/Party/Private/Party/PartyPlatformSessionMonitor.cpp
--- a/Engine/Plugins/Online/OnlineFramework/Source/Party/Private/Party/PartyPlatformSessionMonitor.cpp
+++ b/Engine/Plugins/Online/OnlineFramework/Source/Party/Private/Party/PartyPlatformSessionMonitor.cpp
@@ -515,7 +515,12 @@ void FPartyPlatformSessionMonitor::Initialize()
 
 	EvaluateCurrentSession();
 
-	UpdateRecentPlayersOfLocalMembers(MonitoredParty->GetPartyMembers());
+#if PARTY_PLATFORM_SESSIONS_XBL
+	if (CVarXboxMpaEnabled.GetValueOnAnyThread())
+	{
+		UpdateRecentPlayersOfLocalMembers(MonitoredParty->GetPartyMembers());
+	}
+#endif
 }
 
 void FPartyPlatformSessionMonitor::ShutdownInternal()
@@ -844,6 +849,12 @@ bool FPartyPlatformSessionMonitor::ShouldRecordAsRecentPlayer(const FUniqueNetId
 		return false;
 	}
 
+	const FUniqueNetIdRepl PartyMemberPlatformUniqueId = PartyMember->GetRepData().GetPlatformDataUniqueId();
+	if (!PartyMemberPlatformUniqueId.IsValid())
+	{
+		return false;
+	}
+
 	if (*PartyMember->GetRepData().GetPlatformDataUniqueId() == LocalUserId)
 	{
 		return false;
@@ -858,7 +869,11 @@ void FPartyPlatformSessionMonitor::UpdateRecentPlayersOfLocalMembers(const TArra
 	{
 		if (Member->IsLocalPlayer())
 		{
-			UpdateRecentPlayersOfLocalUser(*Member->GetRepData().GetPlatformDataUniqueId(), RecentPlayers);
+			const FUniqueNetIdRepl PartyMemberPlatformUniqueId = Member->GetRepData().GetPlatformDataUniqueId();
+			if (PartyMemberPlatformUniqueId.IsValid())
+			{
+				UpdateRecentPlayersOfLocalUser(*PartyMemberPlatformUniqueId, RecentPlayers);
+			}
 		}
 	}
 }
@@ -893,19 +908,37 @@ void FPartyPlatformSessionMonitor::HandlePartyMemberInitialized(UPartyMember* In
 		AddLocalPlayerToSession(InitializedMember);
 	}
 
-	TArray<UPartyMember*> RecentPlayers{ InitializedMember };
-	UpdateRecentPlayersOfLocalMembers(RecentPlayers);
-
-	if (InitializedMember->IsLocalPlayer())
+#if PARTY_PLATFORM_SESSIONS_XBL
+	if (CVarXboxMpaEnabled.GetValueOnAnyThread())
 	{
-		UpdateRecentPlayersOfLocalUser(*InitializedMember->GetRepData().GetPlatformDataUniqueId(), MonitoredParty->GetPartyMembers());
+		TArray<UPartyMember*> RecentPlayers{ InitializedMember };
+		UpdateRecentPlayersOfLocalMembers(RecentPlayers);
+
+		if (InitializedMember->IsLocalPlayer())
+		{
+			const FUniqueNetIdRepl PartyMemberPlatformUniqueId = InitializedMember->GetRepData().GetPlatformDataUniqueId();
+			if (PartyMemberPlatformUniqueId.IsValid())
+			{
+				UpdateRecentPlayersOfLocalUser(*PartyMemberPlatformUniqueId, MonitoredParty->GetPartyMembers());
+			}
+		}	
 	}
+#endif
 }
 
 void FPartyPlatformSessionMonitor::HandlePartyMemberLeft(UPartyMember* OldMember, const EMemberExitedReason ExitReason)
 {
-	UE_LOG(LogParty, Verbose, TEXT("HandlePartyMemberLeft: PartyMember=%s User=%s ExitReason=%s"),
-		*OldMember->ToDebugString(true), *OldMember->GetRepData().GetPlatformDataUniqueId().ToDebugString(), ToString(ExitReason));
+	const FUniqueNetIdRepl PartyMemberPlatformUniqueId = OldMember->GetRepData().GetPlatformDataUniqueId();
+
+	if (PartyMemberPlatformUniqueId.IsValid())
+	{
+		UE_LOG(LogParty, Verbose, TEXT("HandlePartyMemberLeft: PartyMember=%s User=%s ExitReason=%s"),
+			*OldMember->ToDebugString(true), *PartyMemberPlatformUniqueId.ToDebugString(), ToString(ExitReason));
+	}
+	else
+	{
+		UE_LOG(LogParty, Warning, TEXT("HandlePartyMemberLeft: PartyMember=Unknown ExitReason=%s"), ToString(ExitReason));
+	}
 
 	if (IsTencentPlatform() && OldMember->GetPlatformOssName() == TENCENT_SUBSYSTEM)
 	{
@@ -1271,4 +1304,4 @@ void FPartyPlatformSessionMonitor::HandleSessionFailure(const FUniqueNetId& Loca
 
 		RetryTickerHandle = FTSTicker::GetCoreTicker().AddTicker(FTickerDelegate::CreateSP(this, &FPartyPlatformSessionMonitor::HandleRetryEstablishingSession), EstablishSessionRetryDelay);
 	}
-}
+}
\ No newline at end of file
