Fixed non unity compile errors in NNERuntimeORT, PartyPlatformSessionMonitor and Analysis.h

diff --git a/Engine/Plugins/Experimental/NNE/Source/NNERuntimeORT/Private/NNERuntimeORT.cpp b/Engine/Plugins/Experimental/NNE/Source/NNERuntimeORT/Private/NNERuntimeORT.cpp
--- a/Engine/Plugins/Experimental/NNE/Source/NNERuntimeORT/Private/NNERuntimeORT.cpp
+++ b/Engine/Plugins/Experimental/NNE/Source/NNERuntimeORT/Private/NNERuntimeORT.cpp
@@ -17,6 +17,7 @@ NNX_THIRD_PARTY_INCLUDES_START
 #include "core/session/onnxruntime_cxx_api.h"
 NNX_THIRD_PARTY_INCLUDES_END
 
+#include UE_INLINE_GENERATED_CPP_BY_NAME(NNERuntimeORT)
 
 using namespace NNX;
 
diff --git a/Engine/Plugins/Online/OnlineFramework/Source/Party/Private/Party/PartyPlatformSessionMonitor.cpp b/Engine/Plugins/Online/OnlineFramework/Source/Party/Private/Party/PartyPlatformSessionMonitor.cpp
--- a/Engine/Plugins/Online/OnlineFramework/Source/Party/Private/Party/PartyPlatformSessionMonitor.cpp
+++ b/Engine/Plugins/Online/OnlineFramework/Source/Party/Private/Party/PartyPlatformSessionMonitor.cpp
@@ -10,10 +10,13 @@
 #include "SocialToolkit.h"
 #include "SocialManager.h"
 
-
 #include "OnlineSessionSettings.h"
 #include "OnlineSubsystemUtils.h"
 
+#if PARTY_PLATFORM_SESSIONS_XBL
+#include "Misc/Base64.h"
+#endif
+
 static bool IsTencentPlatform()
 {
 	return USocialManager::GetSocialOssName(ESocialSubsystem::Platform) == TENCENT_SUBSYSTEM;
diff --git a/Engine/Source/Developer/TraceAnalysis/Public/Trace/Analysis.h b/Engine/Source/Developer/TraceAnalysis/Public/Trace/Analysis.h
--- a/Engine/Source/Developer/TraceAnalysis/Public/Trace/Analysis.h
+++ b/Engine/Source/Developer/TraceAnalysis/Public/Trace/Analysis.h
@@ -4,6 +4,8 @@
 
 #include "Containers/Array.h"
 
+class FMessageLog;
+
 namespace UE {
 namespace Trace {
 
