Add ensure and warning log when a local social user is not properly initialized
Remove global map of initialized delegates
Unbind to delegates actually bound to

diff --git a/Engine/Plugins/Online/OnlineFramework/Source/Party/Private/SocialToolkit.cpp b/Engine/Plugins/Online/OnlineFramework/Source/Party/Private/SocialToolkit.cpp
--- a/Engine/Plugins/Online/OnlineFramework/Source/Party/Private/SocialToolkit.cpp
+++ b/Engine/Plugins/Online/OnlineFramework/Source/Party/Private/SocialToolkit.cpp
@@ -673,21 +673,16 @@ void USocialToolkit::OnOwnerLoggedOut()
 			IOnlinePartyPtr PartyInterface = OSS->GetPartyInterface();
 			if (PartyInterface.IsValid())
 			{
-				PartyInterface->ClearOnPartyInviteReceivedDelegates(this);
+				PartyInterface->ClearOnPartyInviteReceivedExDelegates(this);
+				PartyInterface->ClearOnPartyInviteRemovedExDelegates(this);
 				PartyInterface->ClearOnPartyRequestToJoinReceivedDelegates(this);
 				PartyInterface->ClearOnPartyRequestToJoinRemovedDelegates(this);
 			}
 
-			IOnlineUserPtr UserInterface = OSS->GetUserInterface();
-			if (UserInterface.IsValid())
-			{
-				UserInterface->ClearOnQueryUserInfoCompleteDelegates(LocalUserNum, this);
-			}
-
 			IOnlinePresencePtr PresenceInterface = OSS->GetPresenceInterface();
 			if (PresenceInterface.IsValid())
 			{
-				PresenceInterface->ClearOnPresenceArrayUpdatedDelegates(this);
+				PresenceInterface->ClearOnPresenceReceivedDelegates(this);
 			}
 		}
 	}
diff --git a/Engine/Plugins/Online/OnlineFramework/Source/Party/Private/User/SocialUser.cpp b/Engine/Plugins/Online/OnlineFramework/Source/Party/Private/User/SocialUser.cpp
--- a/Engine/Plugins/Online/OnlineFramework/Source/Party/Private/User/SocialUser.cpp
+++ b/Engine/Plugins/Online/OnlineFramework/Source/Party/Private/User/SocialUser.cpp
@@ -115,8 +115,6 @@ const FOnlineUserPresence* USocialUser::FSubsystemUserInfo::GetPresenceInfo() co
 // USocialUser
 //////////////////////////////////////////////////////////////////////////
 
-TMap<TWeakObjectPtr<USocialUser>, TArray<FOnNewSocialUserInitialized>> USocialUser::InitEventsByUser;
-
 USocialUser::USocialUser()
 {}
 
@@ -181,7 +179,7 @@ void USocialUser::RegisterInitCompleteHandler(const FOnNewSocialUserInitialized&
 		}
 		else
 		{
-			InitEventsByUser.FindOrAdd(this).Add(OnInitializationComplete);
+			UserInitializedEvents.Add(OnInitializationComplete);
 		}
 	}
 }
@@ -347,20 +345,21 @@ void USocialUser::TryBroadcastInitializationComplete()
 
 				bIsInitialized = true;
 
-				TArray<FOnNewSocialUserInitialized> InitEvents;
-				if (InitEventsByUser.RemoveAndCopyValue(this, InitEvents))
+				TArray<FOnNewSocialUserInitialized> InitEvents = MoveTemp(UserInitializedEvents);
+				UserInitializedEvents.Reset();
+				for (FOnNewSocialUserInitialized& InitEvent : InitEvents)
 				{
-					for (FOnNewSocialUserInitialized& InitEvent : InitEvents)
-					{
-						InitEvent.ExecuteIfBound(*this);
-					}
+					InitEvent.ExecuteIfBound(*this);
 				}
 			}
 			else
 			{
+				ensureAlwaysMsgf(!IsLocalUser(), TEXT("%s LocalUser %s does not have UserInfo!"), ANSI_TO_TCHAR(__FUNCTION__), *ToDebugString());
+				UE_LOG(LogParty, Warning, TEXT("%s User %s does not have UserInfo!"), ANSI_TO_TCHAR(__FUNCTION__), *ToDebugString());
+
 				// User is invalid with no open queries
 				// Assume that this means the sought user doesn't exist
-				InitEventsByUser.Remove(this);
+				UserInitializedEvents.Reset();
 
 				// Remove Toolkit's reference to the SocialUser, GC will clean it
 				GetOwningToolkit().HandleUserInvalidated(*this);
diff --git a/Engine/Plugins/Online/OnlineFramework/Source/Party/Public/User/SocialUser.h b/Engine/Plugins/Online/OnlineFramework/Source/Party/Public/User/SocialUser.h
--- a/Engine/Plugins/Online/OnlineFramework/Source/Party/Public/User/SocialUser.h
+++ b/Engine/Plugins/Online/OnlineFramework/Source/Party/Public/User/SocialUser.h
@@ -250,7 +250,7 @@ private:
 	TArray<IOnlinePartyJoinInfoConstRef> ReceivedPartyInvites;
 
 	// Initialization delegates that fire only when a specific user has finishing initializing
-	static TMap<TWeakObjectPtr<USocialUser>, TArray<FOnNewSocialUserInitialized>> InitEventsByUser;
+	TArray<FOnNewSocialUserInitialized> UserInitializedEvents;
 
 	mutable FOnNicknameChanged OnSetNicknameCompletedEvent;
 	mutable FPartyInviteResponseEvent OnPartyInviteAcceptedEvent;
