Add configurable timeouts for UUpdateManager startup and cached response delay.

diff --git a/Engine/Plugins/Online/OnlineFramework/Source/Hotfix/Private/UpdateManager.cpp b/Engine/Plugins/Online/OnlineFramework/Source/Hotfix/Private/UpdateManager.cpp
--- a/Engine/Plugins/Online/OnlineFramework/Source/Hotfix/Private/UpdateManager.cpp
+++ b/Engine/Plugins/Online/OnlineFramework/Source/Hotfix/Private/UpdateManager.cpp
@@ -49,7 +49,9 @@ private:
 };
 
 UUpdateManager::UUpdateManager()
-	: HotfixCheckCompleteDelay(0.1f)
+	: UpdateCheckStartDelay(0.2f)
+	, UpdateCheckCachedResponseDelay(0.1f)
+	, HotfixCheckCompleteDelay(0.1f)
 	, UpdateCheckCompleteDelay(0.5f)
 	, HotfixAvailabilityCheckCompleteDelay(0.1f)
 	, UpdateCheckAvailabilityCompleteDelay(0.1f)
@@ -167,7 +169,7 @@ UUpdateManager::EUpdateStartResult UUpdateManager::StartUpdateCheckInternal(cons
 			};
 
 			// Give the UI state widget a chance to start listening for delegates
-			StartCheckInternalTimerHandle = DelayResponse(StartDelegate, 0.2f);
+			StartCheckInternalTimerHandle = DelayResponse(StartDelegate, UpdateCheckStartDelay);
 			Result = EUpdateStartResult::UpdateStarted;
 		}
 		else
@@ -179,7 +181,7 @@ UUpdateManager::EUpdateStartResult UUpdateManager::StartUpdateCheckInternal(cons
 				StartCheckInternalTimerHandle.Reset();
 			};
 
-			StartCheckInternalTimerHandle = DelayResponse(StartDelegate, 0.1f);
+			StartCheckInternalTimerHandle = DelayResponse(StartDelegate, UpdateCheckCachedResponseDelay);
 			Result = EUpdateStartResult::UpdateCached;
 		}
 	}
diff --git a/Engine/Plugins/Online/OnlineFramework/Source/Hotfix/Public/UpdateManager.h b/Engine/Plugins/Online/OnlineFramework/Source/Hotfix/Public/UpdateManager.h
--- a/Engine/Plugins/Online/OnlineFramework/Source/Hotfix/Public/UpdateManager.h
+++ b/Engine/Plugins/Online/OnlineFramework/Source/Hotfix/Public/UpdateManager.h
@@ -287,6 +287,12 @@ protected:
 	virtual void PostInitProperties() override;
 	virtual void PostReloadConfig(FProperty* PropertyThatWasLoaded) override;
 
+	/** Amount of time to wait before starting an update check. */
+	UPROPERTY(Config)
+	float UpdateCheckStartDelay;
+	/** Amount of time to wait before returning a cached response when updates are started. */
+	UPROPERTY(Config)
+	float UpdateCheckCachedResponseDelay;
 	/** Amount of time to wait between the internal hotfix check completing and advancing to the next stage */
 	UPROPERTY(Config)
 	float HotfixCheckCompleteDelay;
