Creating a private version of EXPOSE_USTRUCT_REP_DATA_PROPERTY_DIRECT

diff --git a/Engine/Plugins/Online/OnlineFramework/Source/Party/Public/Party/PartyTypes.h b/Engine/Plugins/Online/OnlineFramework/Source/Party/Public/Party/PartyTypes.h
--- a/Engine/Plugins/Online/OnlineFramework/Source/Party/Public/Party/PartyTypes.h
+++ b/Engine/Plugins/Online/OnlineFramework/Source/Party/Public/Party/PartyTypes.h
@@ -548,6 +548,10 @@ private: //
 EXPOSE_REP_DATA_PROPERTY_NO_SETTER(Owner, PropertyType, PropertyName, PropertyAccess);	\
 EXPOSE_REP_DATA_PROPERTY_SETTER_ONLY(Owner, PropertyType, PropertyName, PropertyAccess, SetterPrivacy)
 
+/** Expose ustruct property for access without using StructProperty in the property name. */
+#define EXPOSE_PRIVATE_USTRUCT_REP_DATA_PROPERTY_DIRECT(Owner, PropertyType, StructProperty, ChildProperty)	\
+EXPOSE_REP_DATA_PROPERTY_SETTER_ACCESS(Owner, PropertyType, ChildProperty, StructProperty.ChildProperty, private);
+
 /** Expose ustruct property for access without using StructProperty in the property name. */
 #define EXPOSE_USTRUCT_REP_DATA_PROPERTY_DIRECT(Owner, PropertyType, StructProperty, ChildProperty)	\
 EXPOSE_REP_DATA_PROPERTY_SETTER_ACCESS(Owner, PropertyType, ChildProperty, StructProperty.ChildProperty, public);
