- Add new delegates which fire when a hotfix file is added/changed/removed
- Added FileHeader param to PreProcessDownloadedFileData

diff --git a/Engine/Plugins/Online/OnlineFramework/Source/Hotfix/Private/OnlineHotfixManager.cpp b/Engine/Plugins/Online/OnlineFramework/Source/Hotfix/Private/OnlineHotfixManager.cpp
--- a/Engine/Plugins/Online/OnlineFramework/Source/Hotfix/Private/OnlineHotfixManager.cpp
+++ b/Engine/Plugins/Online/OnlineFramework/Source/Hotfix/Private/OnlineHotfixManager.cpp
@@ -435,6 +435,11 @@ void UOnlineHotfixManager::OnEnumerateFilesComplete(bool bWasSuccessful, const F
 							// Perform any undo operations needed
 							RestoreBackupIniFiles();
 							UnmountHotfixFiles();
+
+							for (const FCloudFileHeader& FileHeader : RemovedHotfixFileList)
+							{
+								TriggerOnHotfixRemovedFileDelegates(FileHeader.FileName);
+							}
 							
 							TriggerHotfixComplete(EHotfixResult::SuccessNoChange);
 						}));
@@ -701,6 +706,11 @@ EHotfixResult UOnlineHotfixManager::ApplyHotfix()
 	RestoreBackupIniFiles();
 	UnmountHotfixFiles();
 
+	for (const FCloudFileHeader& FileHeader : RemovedHotfixFileList)
+	{
+		TriggerOnHotfixRemovedFileDelegates(FileHeader.FileName);
+	}
+
 	for (const FCloudFileHeader& FileHeader : ChangedHotfixFileList)
 	{
 		if (!ApplyHotfixProcessing(FileHeader))
@@ -711,6 +721,7 @@ EHotfixResult UOnlineHotfixManager::ApplyHotfix()
 		// Let anyone listening know we just processed this file
 		TriggerOnHotfixProcessedFileDelegates(FileHeader.FileName, GetCachedDirectory() / FileHeader.DLName);
 	}
+
 	UE_LOG(LogHotfixManager, Display, TEXT("Hotfix data has been successfully applied"));
 	EHotfixResult Result = EHotfixResult::Success;
 	if (ChangedOrRemovedPakCount > 0)
@@ -790,20 +801,29 @@ bool UOnlineHotfixManager::ApplyHotfixProcessing(const FCloudFileHeader& FileHea
 
 	bool bSuccess = false;
 	const FString Extension = FPaths::GetExtension(FileHeader.FileName);
-	if (Extension == TEXT("INI"))
+	if (Extension == TEXT("PAK"))
+	{
+		bSuccess = HotfixPakFile(FileHeader);
+	}
+	else
 	{
 		TArray<uint8> FileData;
 		if (OnlineTitleFile->GetFileContents(FileHeader.DLName, FileData))
 		{
-			UE_LOG(LogHotfixManager, Log, TEXT("Applying hotfix %s"), *FileHeader.FileName);
-
-			if (PreProcessDownloadedFileData(FileData))
+			if (PreProcessDownloadedFileData(FileHeader, FileData))
 			{
-				// Convert to a FString
-				FileData.Add(0);
-				FString HotfixStr;
-				FFileHelper::BufferToString(HotfixStr, FileData.GetData(), FileData.Num());
-				bSuccess = HotfixIniFile(FileHeader.FileName, HotfixStr);
+				TriggerOnHotfixUpdatedFileDelegates(FileHeader.FileName, FileData);
+
+				if (Extension == TEXT("INI"))
+				{
+					UE_LOG(LogHotfixManager, Log, TEXT("Applying hotfix %s"), *FileHeader.FileName);
+
+					// Convert to a FString
+					FileData.Add(0);
+					FString HotfixStr;
+					FFileHelper::BufferToString(HotfixStr, FileData.GetData(), FileData.Num());
+					bSuccess = HotfixIniFile(FileHeader.FileName, HotfixStr);
+				}
 			}
 			else
 			{
@@ -815,10 +835,6 @@ bool UOnlineHotfixManager::ApplyHotfixProcessing(const FCloudFileHeader& FileHea
 			UE_LOG(LogHotfixManager, Warning, TEXT("Failed to get contents of %s"), *FileHeader.FileName);
 		}
 	}
-	else if (Extension == TEXT("PAK"))
-	{
-		bSuccess = HotfixPakFile(FileHeader);
-	}
 	OnlineTitleFile->ClearFile(FileHeader.FileName);
 	return bSuccess;
 }
diff --git a/Engine/Plugins/Online/OnlineFramework/Source/Hotfix/Public/OnlineHotfixManager.h b/Engine/Plugins/Online/OnlineFramework/Source/Hotfix/Public/OnlineHotfixManager.h
--- a/Engine/Plugins/Online/OnlineFramework/Source/Hotfix/Public/OnlineHotfixManager.h
+++ b/Engine/Plugins/Online/OnlineFramework/Source/Hotfix/Public/OnlineHotfixManager.h
@@ -54,7 +54,7 @@ DECLARE_MULTICAST_DELEGATE_FourParams(FOnHotfixProgress, uint32, uint32, uint64,
 typedef FOnHotfixProgress::FDelegate FOnHotfixProgressDelegate;
 
 /**
- * Delegate fired as each file is applied
+ * Delegate fired for each new/updated file after it is applied
  *
  * @param FriendlyName the human readable version of the file name (DefaultEngine.ini)
  * @param CachedFileName the full path to the file on disk
@@ -62,6 +62,23 @@ typedef FOnHotfixProgress::FDelegate FOnHotfixProgressDelegate;
 DECLARE_MULTICAST_DELEGATE_TwoParams(FOnHotfixProcessedFile, const FString&, const FString&);
 typedef FOnHotfixProcessedFile::FDelegate FOnHotfixProcessedFileDelegate;
 
+/**
+ * Delegate fired for each removed file
+ *
+ * @param FriendlyName the human readable version of the file name (DefaultEngine.ini)
+ */
+DECLARE_MULTICAST_DELEGATE_OneParam(FOnHotfixRemovedFile, const FString&);
+typedef FOnHotfixRemovedFile::FDelegate FOnHotfixRemovedFileDelegate;
+
+/**
+ * Delegate fired for each added/updated file
+ *
+ * @param FriendlyName the human readable version of the file name (DefaultEngine.ini)
+ * @param FileContents the preprocessed contents of the file.
+ */
+DECLARE_MULTICAST_DELEGATE_TwoParams(FOnHotfixUpdatedFile, const FString&, const TArray<uint8>&);
+typedef FOnHotfixUpdatedFile::FDelegate FOnHotfixUpdatedFileDelegate;
+
 /**
  * This class manages the downloading and application of hotfix data
  * Hotfix data is a set of non-executable files downloaded and applied to the game.
@@ -100,10 +117,20 @@ protected:
 	DEFINE_ONLINE_DELEGATE_FOUR_PARAM(OnHotfixProgress, uint32, uint32, uint64, uint64);
 
 	/**
-	 * Delegate fired as the hotfix files are applied
+	 * Delegate fired for each new/updated file after it is applied
 	 */
 	DEFINE_ONLINE_DELEGATE_TWO_PARAM(OnHotfixProcessedFile, const FString&, const FString&);
 
+	/**
+	 * Delegate fired for each removed file
+	 */
+	DEFINE_ONLINE_DELEGATE_ONE_PARAM(OnHotfixRemovedFile, const FString&);
+
+	/**
+	 * Delegate fired for each added/updated file
+	 */
+	DEFINE_ONLINE_DELEGATE_TWO_PARAM(OnHotfixUpdatedFile, const FString&, const TArray<uint8>&);
+
 	struct FPendingFileDLProgress
 	{
 		uint64 Progress;
@@ -248,13 +275,23 @@ protected:
 	 */
 	virtual bool ApplyHotfixProcessing(const FCloudFileHeader& FileHeader);
 	/**
-	* Called prior to reading the file data.
-	*
-	* @param FileData - byte data of the hotfix file
-	*
-	* @return whether the file was successfully preprocessed
-	*/
+	 * Called prior to reading the file data.
+	 *
+	 * @param FileData - byte data of the hotfix file
+	 *
+	 * @return whether the file was successfully preprocessed
+	 */
+	UE_DEPRECATED(5.3, "Replaced with PreProcessDownloadedFileData taking a FCloudFileHeader")
 	virtual bool PreProcessDownloadedFileData(TArray<uint8>& FileData) const { return true;	}
+	/**
+	 * Called prior to reading the file data.
+	 *
+	 * @param FileHeader - the header information for the file in question
+	 * @param FileData - byte data of the hotfix file. Intentionally not const, so the array is modifiable as part of preprocessing.
+	 *
+	 * @return whether the file was successfully preprocessed
+	 */
+	virtual bool PreProcessDownloadedFileData(const FCloudFileHeader& FileHeader, TArray<uint8>& FileData) const { return true; }
 	/**
 	 * Override this to change the default INI file handling (merge delta INI changes into the config cache)
 	 *
