- added logging to help QA verify changes are getting downloaded by OnlineHotfixManager

diff --git a/Engine/Plugins/Online/OnlineFramework/Source/Hotfix/Private/OnlineHotfixManager.cpp b/Engine/Plugins/Online/OnlineFramework/Source/Hotfix/Private/OnlineHotfixManager.cpp
--- a/Engine/Plugins/Online/OnlineFramework/Source/Hotfix/Private/OnlineHotfixManager.cpp
+++ b/Engine/Plugins/Online/OnlineFramework/Source/Hotfix/Private/OnlineHotfixManager.cpp
@@ -618,6 +618,10 @@ void UOnlineHotfixManager::BuildHotfixFileListDeltas()
 			}
 		}
 	}
+
+	UE_LOG(LogHotfixManager, Log, TEXT("Hotfix deltas: updated:[%s], removed:[%s]"),
+		*FString::JoinBy(ChangedHotfixFileList, TEXT(", "), &FCloudFileHeader::FileName),
+		*FString::JoinBy(RemovedHotfixFileList, TEXT(", "), &FCloudFileHeader::FileName));
 }
 
 void UOnlineHotfixManager::FilterHotfixFiles()
