Add support for social debug tools to login multiple golems simultaneously

diff --git a/Engine/Plugins/Online/OnlineFramework/Source/Party/Private/SocialDebugTools.cpp b/Engine/Plugins/Online/OnlineFramework/Source/Party/Private/SocialDebugTools.cpp
--- a/Engine/Plugins/Online/OnlineFramework/Source/Party/Private/SocialDebugTools.cpp
+++ b/Engine/Plugins/Online/OnlineFramework/Source/Party/Private/SocialDebugTools.cpp
@@ -35,7 +35,7 @@ void USocialDebugTools::PrintExecUsage() const
 
 void USocialDebugTools::PrintExecCommands() const
 {
-	UE_LOG(LogParty, Log, TEXT("LOGIN <Id> <Auth>"));
+	UE_LOG(LogParty, Log, TEXT("LOGIN <Id> <Auth> <optional CONTEXT=<Name2> <Id2> <Auth2> CONTEXT=<Name3> <Id3> <Auth3> ...>"));
 	UE_LOG(LogParty, Log, TEXT("LOGOUT"));
 	UE_LOG(LogParty, Log, TEXT("JOINPARTY <Id> <Auth> <optional FriendName>"));
 	UE_LOG(LogParty, Log, TEXT("LEAVEPARTY"));
@@ -536,14 +536,32 @@ bool USocialDebugTools::RunCommand(const TCHAR* Cmd, const TArray<FString>& Targ
 		}
 		else
 		{
-			const FString Id = FParse::Token(Cmd, false);
-			const FString Auth = FParse::Token(Cmd, false);
 			FString InstanceName = TargetInstances[0];
-
-			Login(InstanceName, FOnlineAccountCredentials(TEXT("epic"), Id, Auth), FLoginComplete::CreateLambda([this, InstanceName](bool bSuccess)
+			FString Id = FParse::Token(Cmd, false);
+			FString Auth = FParse::Token(Cmd, false);
+			
+			int MaxConcurrentLogins = 512; // Reasonable limit to prevent infinite loop
+			while (!Id.IsEmpty() && !Auth.IsEmpty() && MaxConcurrentLogins >= 0)
 			{
-				UE_LOG(LogParty, Display, TEXT("Login OSS context[%s] %s"), *InstanceName, *LexToString(bSuccess));
-			}));
+				Login(InstanceName, FOnlineAccountCredentials(TEXT("epic"), Id, Auth), FLoginComplete::CreateLambda([this, InstanceName](bool bSuccess)
+				{
+					UE_LOG(LogParty, Display, TEXT("Login OSS context[%s] %s"), *InstanceName, *LexToString(bSuccess));
+				}));
+				
+				// Allow for more logins by concatenating more context=name auth pass
+				if (FParse::Value(Cmd, TEXT("CONTEXT="), InstanceName))
+				{
+					// strip out context to parse next entry
+					FParse::Command(&Cmd, *FString(TEXT("CONTEXT=") + InstanceName));
+					Id = FParse::Token(Cmd, false);
+					Auth = FParse::Token(Cmd, false);
+					--MaxConcurrentLogins;
+				}
+				else
+				{
+					break;
+				}
+			}
 		}
 		return true;
 	}
