Slight refactor to social debug tools to reduce boilerplate

diff --git a/Engine/Plugins/Online/OnlineFramework/Source/Party/Private/SocialDebugTools.cpp b/Engine/Plugins/Online/OnlineFramework/Source/Party/Private/SocialDebugTools.cpp
--- a/Engine/Plugins/Online/OnlineFramework/Source/Party/Private/SocialDebugTools.cpp
+++ b/Engine/Plugins/Online/OnlineFramework/Source/Party/Private/SocialDebugTools.cpp
@@ -205,80 +205,76 @@ void USocialDebugTools::Logout(const FString& Instance, const FLogoutComplete& O
 
 void USocialDebugTools::JoinParty(const FString& Instance, const FString& FriendName, const FJoinPartyComplete& OnComplete)
 {
-	IOnlineSubsystem* OnlineSub = GetContext(Instance).GetOSS();
+	FInstanceContext& Context = GetContext(Instance);
+	IOnlineSubsystem* OnlineSub = Context.GetOSS();
 	if (OnlineSub)
 	{
-		IOnlineIdentityPtr OnlineIdentity = OnlineSub->GetIdentityInterface();
-		if (OnlineIdentity.IsValid())
+		if (FUniqueNetIdPtr LocalUserId = Context.GetLocalUserId())
 		{
-			FUniqueNetIdPtr LocalUserId = OnlineIdentity->GetUniquePlayerId(LocalUserNum);
-			if (LocalUserId.IsValid())
+			IOnlinePartyPtr OnlineParty = OnlineSub->GetPartyInterface();
+			if (OnlineParty.IsValid())
 			{
-				IOnlinePartyPtr OnlineParty = OnlineSub->GetPartyInterface();
-				if (OnlineParty.IsValid())
+				// query user id by name
+				if (!FriendName.IsEmpty())
 				{
-					// query user id by name
-					if (!FriendName.IsEmpty())
+					IOnlineUserPtr OnlineUser = OnlineSub->GetUserInterface();
+					if (OnlineUser.IsValid())
 					{
-						IOnlineUserPtr OnlineUser = OnlineSub->GetUserInterface();
-						if (OnlineUser.IsValid())
+						OnlineUser->QueryUserIdMapping(*LocalUserId, FriendName, IOnlineUser::FOnQueryUserMappingComplete::CreateLambda([this, Instance, OnlineParty, OnComplete](bool bWasSuccessful, const FUniqueNetId& UserId, const FString& DisplayNameOrEmail, const FUniqueNetId& FoundUserId, const FString& Error)
 						{
-							OnlineUser->QueryUserIdMapping(*LocalUserId, FriendName, IOnlineUser::FOnQueryUserMappingComplete::CreateLambda([this, Instance, OnlineParty, OnComplete](bool bWasSuccessful, const FUniqueNetId& UserId, const FString& DisplayNameOrEmail, const FUniqueNetId& FoundUserId, const FString& Error)
+							if (bWasSuccessful)
 							{
-								if (bWasSuccessful)
+								IOnlinePartyJoinInfoConstPtr JoinInfo = OnlineParty->GetAdvertisedParty(UserId, FoundUserId, IOnlinePartySystem::GetPrimaryPartyTypeId());
+								if (JoinInfo.IsValid())
 								{
-									IOnlinePartyJoinInfoConstPtr JoinInfo = OnlineParty->GetAdvertisedParty(UserId, FoundUserId, IOnlinePartySystem::GetPrimaryPartyTypeId());
-									if (JoinInfo.IsValid())
+									OnlineParty->JoinParty(UserId, *JoinInfo, FOnJoinPartyComplete::CreateLambda([this, Instance, OnlineParty, OnComplete](const FUniqueNetId& UserIdTmp, const FOnlinePartyId& PartyId, const EJoinPartyCompletionResult Result, const int32 NotApprovedReason)
 									{
-										OnlineParty->JoinParty(UserId, *JoinInfo, FOnJoinPartyComplete::CreateLambda([this, Instance, OnlineParty, OnComplete](const FUniqueNetId& UserIdTmp, const FOnlinePartyId& PartyId, const EJoinPartyCompletionResult Result, const int32 NotApprovedReason)
+										bool bSuccess = Result == EJoinPartyCompletionResult::Succeeded;
+										if (bSuccess)
 										{
-											bool bSuccess = Result == EJoinPartyCompletionResult::Succeeded;
-											if (bSuccess)
-											{
-												FOnlinePartyDataConstPtr PartyMemberData = GetContext(Instance).GetPartyMemberData();
-												if (PartyMemberData.IsValid())
-												{
-													OnlineParty->UpdatePartyMemberData(UserIdTmp, PartyId, DefaultPartyDataNamespace, *PartyMemberData);
-												}
-											}
-											else
+											FOnlinePartyDataConstPtr PartyMemberData = GetContext(Instance).GetPartyMemberData();
+											if (PartyMemberData.IsValid())
 											{
-												UE_LOG(LogParty, Warning, TEXT("Party context[%s] join attempt denied for reason [%d]"), *Instance, (int32)Result);
+												OnlineParty->UpdatePartyMemberData(UserIdTmp, PartyId, DefaultPartyDataNamespace, *PartyMemberData);
 											}
-											OnComplete.ExecuteIfBound(bSuccess);
-										}));
-										return;
-									}
+										}
+										else
+										{
+											UE_LOG(LogParty, Warning, TEXT("Party context[%s] join attempt denied for reason [%d]"), *Instance, (int32)Result);
+										}
+										OnComplete.ExecuteIfBound(bSuccess);
+									}));
+									return;
 								}
-								OnComplete.ExecuteIfBound(false);
-							}));
-							return;
-						}
-					}					
-					else
-					{	
-						IOnlinePartyJoinInfoConstPtr JoinInfo = GetDefaultPartyJoinInfo();
-						if (JoinInfo.IsValid())
+							}
+							OnComplete.ExecuteIfBound(false);
+						}));
+						return;
+					}
+				}					
+				else
+				{	
+					IOnlinePartyJoinInfoConstPtr JoinInfo = GetDefaultPartyJoinInfo();
+					if (JoinInfo.IsValid())
+					{
+						OnlineParty->JoinParty(*LocalUserId, *JoinInfo, FOnJoinPartyComplete::CreateLambda([this, Instance, OnlineParty, OnComplete](const FUniqueNetId& UserId, const FOnlinePartyId& PartyId, const EJoinPartyCompletionResult Result, const int32 NotApprovedReason)
 						{
-							OnlineParty->JoinParty(*LocalUserId, *JoinInfo, FOnJoinPartyComplete::CreateLambda([this, Instance, OnlineParty, OnComplete](const FUniqueNetId& UserId, const FOnlinePartyId& PartyId, const EJoinPartyCompletionResult Result, const int32 NotApprovedReason)
+							bool bSuccess = Result == EJoinPartyCompletionResult::Succeeded;
+							if (bSuccess)
 							{
-								bool bSuccess = Result == EJoinPartyCompletionResult::Succeeded;
-								if (bSuccess)
-								{
-									FOnlinePartyDataConstPtr PartyMemberData = GetContext(Instance).GetPartyMemberData();
-									if (PartyMemberData.IsValid())
-									{
-										OnlineParty->UpdatePartyMemberData(UserId, PartyId, DefaultPartyDataNamespace, *PartyMemberData);
-									}
-								}
-								else
+								FOnlinePartyDataConstPtr PartyMemberData = GetContext(Instance).GetPartyMemberData();
+								if (PartyMemberData.IsValid())
 								{
-									UE_LOG(LogParty, Warning, TEXT("Party context[%s] join attempt denied for reason [%d]"), *Instance, (int32)Result);
+									OnlineParty->UpdatePartyMemberData(UserId, PartyId, DefaultPartyDataNamespace, *PartyMemberData);
 								}
-								OnComplete.ExecuteIfBound(bSuccess);
-							}));
-							return;
-						}
+							}
+							else
+							{
+								UE_LOG(LogParty, Warning, TEXT("Party context[%s] join attempt denied for reason [%d]"), *Instance, (int32)Result);
+							}
+							OnComplete.ExecuteIfBound(bSuccess);
+						}));
+						return;
 					}
 				}
 			}
@@ -289,27 +285,23 @@ void USocialDebugTools::JoinParty(const FString& Instance, const FString& Friend
 
 void USocialDebugTools::LeaveParty(const FString& Instance, const FLeavePartyComplete& OnComplete)
 {
+	FInstanceContext& Context = GetContext(Instance);
 	IOnlineSubsystem* OnlineSub = GetContext(Instance).GetOSS();
 	if (OnlineSub)
 	{
-		IOnlineIdentityPtr OnlineIdentity = OnlineSub->GetIdentityInterface();
-		if (OnlineIdentity.IsValid())
+		if (FUniqueNetIdPtr UserId = Context.GetLocalUserId())
 		{
-			FUniqueNetIdPtr UserId = OnlineIdentity->GetUniquePlayerId(LocalUserNum);
-			if (UserId.IsValid())
+			IOnlinePartyPtr OnlineParty = OnlineSub->GetPartyInterface();
+			if (OnlineParty.IsValid())
 			{
-				IOnlinePartyPtr OnlineParty = OnlineSub->GetPartyInterface();
-				if (OnlineParty.IsValid())
+				TSharedPtr<const FOnlineParty> ExistingParty = OnlineParty->GetParty(*UserId, IOnlinePartySystem::GetPrimaryPartyTypeId());
+				if (ExistingParty.IsValid())
 				{
-					TSharedPtr<const FOnlineParty> ExistingParty = OnlineParty->GetParty(*UserId, IOnlinePartySystem::GetPrimaryPartyTypeId());
-					if (ExistingParty.IsValid())
+					OnlineParty->LeaveParty(*UserId, *ExistingParty->PartyId, FOnLeavePartyComplete::CreateLambda([this, Instance, OnComplete](const FUniqueNetId& LocalUserId, const FOnlinePartyId& PartyId, const ELeavePartyCompletionResult Result)
 					{
-						OnlineParty->LeaveParty(*UserId, *ExistingParty->PartyId, FOnLeavePartyComplete::CreateLambda([this, Instance, OnComplete](const FUniqueNetId& LocalUserId, const FOnlinePartyId& PartyId, const ELeavePartyCompletionResult Result)
-						{
-							OnComplete.ExecuteIfBound(Result == ELeavePartyCompletionResult::Succeeded);
-						}));
-						return;
-					}
+						OnComplete.ExecuteIfBound(Result == ELeavePartyCompletionResult::Succeeded);
+					}));
+					return;
 				}
 			}
 		}
@@ -319,24 +311,20 @@ void USocialDebugTools::LeaveParty(const FString& Instance, const FLeavePartyCom
 
 void USocialDebugTools::CleanupParties(const FString& Instance, const FCleanupPartiesComplete& OnComplete)
 {
-	IOnlineSubsystem* OnlineSub = GetContext(Instance).GetOSS();
+	FInstanceContext& Context = GetContext(Instance);
+	IOnlineSubsystem* OnlineSub = Context.GetOSS();
 	if (OnlineSub)
 	{
-		IOnlineIdentityPtr OnlineIdentity = OnlineSub->GetIdentityInterface();
-		if (OnlineIdentity.IsValid())
+		if (FUniqueNetIdPtr UserId = Context.GetLocalUserId())
 		{
-			FUniqueNetIdPtr UserId = OnlineIdentity->GetUniquePlayerId(LocalUserNum);
-			if (UserId.IsValid())
+			IOnlinePartyPtr OnlineParty = OnlineSub->GetPartyInterface();
+			if (OnlineParty.IsValid())
 			{
-				IOnlinePartyPtr OnlineParty = OnlineSub->GetPartyInterface();
-				if (OnlineParty.IsValid())
+				OnlineParty->CleanupParties(*UserId, FOnCleanupPartiesComplete::CreateLambda([this, Instance, OnComplete](const FUniqueNetId& LocalUserId, const FOnlineError& Result)
 				{
-					OnlineParty->CleanupParties(*UserId, FOnCleanupPartiesComplete::CreateLambda([this, Instance, OnComplete](const FUniqueNetId& LocalUserId, const FOnlineError& Result)
-					{
-						OnComplete.ExecuteIfBound(Result.WasSuccessful());
-					}));
-					return;
-				}
+					OnComplete.ExecuteIfBound(Result.WasSuccessful());
+				}));
+				return;
 			}
 		}
 	}
@@ -364,24 +352,19 @@ void USocialDebugTools::SetPartyMemberData(const FString& Instance, const UStruc
 		IOnlineSubsystem* OnlineSub = Context.GetOSS();
 		if (OnlineSub)
 		{
-			IOnlineIdentityPtr OnlineIdentity = OnlineSub->GetIdentityInterface();
-			if (OnlineIdentity.IsValid())
+			if (FUniqueNetIdPtr UserId = Context.GetLocalUserId())
 			{
-				FUniqueNetIdPtr UserId = OnlineIdentity->GetUniquePlayerId(LocalUserNum);
-				if (UserId.IsValid())
+				IOnlinePartyPtr OnlineParty = OnlineSub->GetPartyInterface();
+				if (OnlineParty.IsValid())
 				{
-					IOnlinePartyPtr OnlineParty = OnlineSub->GetPartyInterface();
-					if (OnlineParty.IsValid())
+					TSharedPtr<const FOnlineParty> ExistingParty = OnlineParty->GetParty(*UserId, IOnlinePartySystem::GetPrimaryPartyTypeId());
+					if (ExistingParty.IsValid())
 					{
-						TSharedPtr<const FOnlineParty> ExistingParty = OnlineParty->GetParty(*UserId, IOnlinePartySystem::GetPrimaryPartyTypeId());
-						if (ExistingParty.IsValid())
+						const FOnlinePartyId& PartyId = *ExistingParty->PartyId;
+						if (StructType->IsChildOf(FPartyMemberRepData::StaticStruct()))
 						{
-							const FOnlinePartyId& PartyId = *ExistingParty->PartyId;
-							if (StructType->IsChildOf(FPartyMemberRepData::StaticStruct()))
-							{
-								UE_LOG(LogParty, Display, TEXT("Sending rep data update for member within party [%s]."), *PartyId.ToDebugString());
-								bResult = OnlineParty->UpdatePartyMemberData(*UserId, PartyId, DefaultPartyDataNamespace, *Context.PartyMemberData);
-							}
+							UE_LOG(LogParty, Display, TEXT("Sending rep data update for member within party [%s]."), *PartyId.ToDebugString());
+							bResult = OnlineParty->UpdatePartyMemberData(*UserId, PartyId, DefaultPartyDataNamespace, *Context.PartyMemberData);
 						}
 					}
 				}
@@ -410,22 +393,17 @@ void USocialDebugTools::SetPartyMemberDataJson(const FString& Instance, const FS
 		IOnlineSubsystem* OnlineSub = Context.GetOSS();
 		if (OnlineSub)
 		{
-			IOnlineIdentityPtr OnlineIdentity = OnlineSub->GetIdentityInterface();
-			if (OnlineIdentity.IsValid())
+			if (FUniqueNetIdPtr UserId = Context.GetLocalUserId())
 			{
-				FUniqueNetIdPtr UserId = OnlineIdentity->GetUniquePlayerId(LocalUserNum);
-				if (UserId.IsValid())
+				IOnlinePartyPtr OnlineParty = OnlineSub->GetPartyInterface();
+				if (OnlineParty.IsValid())
 				{
-					IOnlinePartyPtr OnlineParty = OnlineSub->GetPartyInterface();
-					if (OnlineParty.IsValid())
+					TSharedPtr<const FOnlineParty> ExistingParty = OnlineParty->GetParty(*UserId, IOnlinePartySystem::GetPrimaryPartyTypeId());
+					if (ExistingParty.IsValid())
 					{
-						TSharedPtr<const FOnlineParty> ExistingParty = OnlineParty->GetParty(*UserId, IOnlinePartySystem::GetPrimaryPartyTypeId());
-						if (ExistingParty.IsValid())
-						{
-							const FOnlinePartyId& PartyId = *ExistingParty->PartyId;
-							UE_LOG(LogParty, Display, TEXT("Sending rep data update for member within party [%s]."), *PartyId.ToDebugString());
-							bResult = OnlineParty->UpdatePartyMemberData(*UserId, PartyId, DefaultPartyDataNamespace, *Context.PartyMemberData);
-						}
+						const FOnlinePartyId& PartyId = *ExistingParty->PartyId;
+						UE_LOG(LogParty, Display, TEXT("Sending rep data update for member within party [%s]."), *PartyId.ToDebugString());
+						bResult = OnlineParty->UpdatePartyMemberData(*UserId, PartyId, DefaultPartyDataNamespace, *Context.PartyMemberData);
 					}
 				}
 			}
@@ -448,7 +426,6 @@ IOnlineSubsystem* USocialDebugTools::GetDefaultOSS() const
 	{
 		return IOnlineSubsystem::Get(MCP_SUBSYSTEM);
 	}
-	
 }
 
 USocialManager& USocialDebugTools::GetSocialManager() const
@@ -510,12 +487,9 @@ USocialDebugTools::FInstanceContext* USocialDebugTools::GetContextForUser(const
 		FInstanceContext& Context = Entry.Value;
 		if (Context.GetOSS())
 		{
-			IOnlineIdentityPtr OnlineIdentity = Context.GetOSS()->GetIdentityInterface();
-			if (OnlineIdentity.IsValid())
+			if (FUniqueNetIdPtr LocalUserId = Context.GetLocalUserId())
 			{
-				FUniqueNetIdPtr LocalUser = OnlineIdentity->GetUniquePlayerId(LocalUserNum);
-				if (LocalUser.IsValid() &&
-					*LocalUser == UserId)
+				if (*LocalUserId == UserId)
 				{
 					Result = &Context;
 					break;
@@ -708,6 +682,20 @@ void USocialDebugTools::FInstanceContext::Shutdown()
 	}
 }
 
+FUniqueNetIdPtr USocialDebugTools::FInstanceContext::GetLocalUserId() const
+{
+	FUniqueNetIdPtr LocalUserId;
+	if (OnlineSub)
+	{
+		IOnlineIdentityPtr OnlineIdentity = OnlineSub->GetIdentityInterface();
+		if (OnlineIdentity.IsValid())
+		{
+			LocalUserId = OnlineIdentity->GetUniquePlayerId(LocalUserNum);
+		}
+	}
+	return LocalUserId;
+}
+
 void USocialDebugTools::HandleFriendInviteReceived(const FUniqueNetId& LocalUserId, const FUniqueNetId& FriendId)
 {
 	if (bAutoAcceptFriendInvites)
diff --git a/Engine/Plugins/Online/OnlineFramework/Source/Party/Public/SocialDebugTools.h b/Engine/Plugins/Online/OnlineFramework/Source/Party/Public/SocialDebugTools.h
--- a/Engine/Plugins/Online/OnlineFramework/Source/Party/Public/SocialDebugTools.h
+++ b/Engine/Plugins/Online/OnlineFramework/Source/Party/Public/SocialDebugTools.h
@@ -15,13 +15,14 @@ class IOnlinePartyPendingJoinRequestInfo;
 typedef TSharedPtr<const class IOnlinePartyJoinInfo> IOnlinePartyJoinInfoConstPtr;
 typedef TSharedPtr<class FOnlinePartyData> FOnlinePartyDataPtr;
 typedef TSharedPtr<const class FOnlinePartyData> FOnlinePartyDataConstPtr;
+using FUniqueNetIdPtr = TSharedPtr<const FUniqueNetId>;
 
 UCLASS(Within = SocialManager, Config = Game)
 class PARTY_API USocialDebugTools : public UObject, public FExec
 {
 	GENERATED_BODY()
 
-	static const int32 LocalUserNum = 0;
+	static constexpr const int32 LocalUserNum = 0;
 
 public:
 	USocialManager& GetSocialManager() const;
@@ -65,8 +66,9 @@ public:
 
 		void Init();
 		void Shutdown();
-		inline IOnlineSubsystem* GetOSS() { return OnlineSub; }
-		inline FOnlinePartyDataConstPtr GetPartyMemberData() { return PartyMemberData; }
+		IOnlineSubsystem* GetOSS() const { return OnlineSub; }
+		FOnlinePartyDataPtr GetPartyMemberData() const { return PartyMemberData; }
+		FUniqueNetIdPtr GetLocalUserId() const;
 
 		FString Name;
 		IOnlineSubsystem* OnlineSub;
