Respond to JIP request with GameFull when the party beacon returns an error of that nature

diff --git a/Engine/Plugins/Online/OnlineFramework/Source/Party/Private/Party/SocialParty.cpp b/Engine/Plugins/Online/OnlineFramework/Source/Party/Private/Party/SocialParty.cpp
--- a/Engine/Plugins/Online/OnlineFramework/Source/Party/Private/Party/SocialParty.cpp
+++ b/Engine/Plugins/Online/OnlineFramework/Source/Party/Private/Party/SocialParty.cpp
@@ -771,7 +771,7 @@ void USocialParty::HandlePartyJoinRequestReceived(const FUniqueNetId& LocalUserI
 	else
 	{
 		const bool bIsApproved = JoinApproval.CanJoin();
-		UE_LOG(LogParty, Verbose, TEXT("[%s] Responding to approval request for %s with %s"), *PartyId.ToString(), *PrimaryJoiningUser->GetUserId()->ToDebugString(), bIsApproved ? TEXT("approved") : TEXT("denied"));
+		UE_LOG(LogParty, Log, TEXT("[%s] Responding to approval request for %s with %s"), *PartyId.ToString(), *PrimaryJoiningUser->GetUserId()->ToDebugString(), bIsApproved ? TEXT("approved") : TEXT("denied"));
 		
 		PartyInterface->ApproveJoinRequest(LocalUserId, PartyId, *PrimaryJoiningUser->GetUserId(), bIsApproved, JoinApproval.GetDenialReason());
 	}
@@ -853,7 +853,7 @@ void USocialParty::HandlePartyJIPRequestReceived(const FUniqueNetId& LocalUserId
 	else
 	{
 		const bool bIsApproved = JoinApproval.CanJoin();
-		UE_LOG(LogParty, Verbose, TEXT("[%s] Responding to approval request for %s with %s"), *PartyId.ToString(), *SenderId.ToString(), bIsApproved ? TEXT("approved") : TEXT("denied"));
+		UE_LOG(LogParty, Log, TEXT("[%s] Responding to approval request for %s with %s"), *PartyId.ToString(), *SenderId.ToString(), bIsApproved ? TEXT("approved") : TEXT("denied"));
 
 		PRAGMA_DISABLE_DEPRECATION_WARNINGS
 		PartyInterface->ApproveJIPRequest(LocalUserId, PartyId, SenderId, bIsApproved, JoinApproval.GetDenialReason());
@@ -1606,7 +1606,7 @@ void USocialParty::RejectAllPendingJoinRequests()
 	{
 		PendingApprovals.Dequeue(PendingApproval);
 		const FUniqueNetId& PrimaryJoiningUserId = *PendingApproval.Members[0].MemberId.GetUniqueNetId();
-		UE_LOG(LogParty, Verbose, TEXT("[%s] Responding to approval request for %s with denied"), *PartyId.ToString(), *PrimaryJoiningUserId.ToDebugString());
+		UE_LOG(LogParty, Log, TEXT("[%s] Responding to approval request for %s with denied"), *PartyId.ToString(), *PrimaryJoiningUserId.ToDebugString());
 		if (PendingApproval.bIsJIPApproval)
 		{
 			PRAGMA_DISABLE_DEPRECATION_WARNINGS
@@ -1700,12 +1700,29 @@ void USocialParty::PumpApprovalQueue()
 	}
 }
 
+namespace UE::Online::Party::Private
+{
+	FPartyJoinDenialReason ReservationResponseToJoinDenialReason(const EPartyReservationResult::Type ReservationResponse)
+	{
+		switch (ReservationResponse)
+		{
+		case EPartyReservationResult::PartyLimitReached:
+			return EPartyJoinDenialReason::GameFull;
+		case EPartyReservationResult::ReservationDenied:
+			return EPartyJoinDenialReason::GameFull; // Not really a more specific reason available
+		case EPartyReservationResult::ReservationDenied_CrossPlayRestriction:
+			return EPartyJoinDenialReason::JoinerCrossplayRestricted;
+		}
+		return EPartyJoinDenialReason::NoReason;
+	}
+}
+
 void USocialParty::HandleReservationRequestComplete(EPartyReservationResult::Type ReservationResponse)
 {
 	UE_LOG(LogParty, Verbose, TEXT("Reservation request complete with response: %s"), EPartyReservationResult::ToString(ReservationResponse));
 
 	const bool bReservationApproved = ReservationResponse == EPartyReservationResult::ReservationAccepted || ReservationResponse == EPartyReservationResult::ReservationDuplicate;
-	const FPartyJoinDenialReason DenialReason = ReservationResponse == EPartyReservationResult::ReservationDenied_CrossPlayRestriction ? EPartyJoinDenialReason::JoinerCrossplayRestricted : EPartyJoinDenialReason::NoReason;
+	const FPartyJoinDenialReason DenialReason = UE::Online::Party::Private::ReservationResponseToJoinDenialReason(ReservationResponse);
 
 	if (bReservationApproved || DenialReason.HasAnyReason())
 	{
