Add JIP command to SocialDebugTools

diff --git a/Engine/Plugins/Online/OnlineFramework/Source/Party/Private/Party/PartyMember.cpp b/Engine/Plugins/Online/OnlineFramework/Source/Party/Private/Party/PartyMember.cpp
--- a/Engine/Plugins/Online/OnlineFramework/Source/Party/Private/Party/PartyMember.cpp
+++ b/Engine/Plugins/Online/OnlineFramework/Source/Party/Private/Party/PartyMember.cpp
@@ -18,9 +18,14 @@ void FPartyMemberRepData::SetOwningMember(const UPartyMember& InOwnerMember)
 	OwnerMember = &InOwnerMember;
 }
 
+void FPartyMemberRepData::MarkOwnerless()
+{
+	bAllowOwnerless = true;
+}
+
 bool FPartyMemberRepData::CanEditData() const
 {
-	return OwnerMember.IsValid() && OwnerMember->IsLocalPlayer();
+	return bAllowOwnerless || (OwnerMember.IsValid() && OwnerMember->IsLocalPlayer());
 }
 
 void FPartyMemberRepData::CompareAgainst(const FOnlinePartyRepDataBase& OldData) const
diff --git a/Engine/Plugins/Online/OnlineFramework/Source/Party/Private/Party/SocialParty.cpp b/Engine/Plugins/Online/OnlineFramework/Source/Party/Private/Party/SocialParty.cpp
--- a/Engine/Plugins/Online/OnlineFramework/Source/Party/Private/Party/SocialParty.cpp
+++ b/Engine/Plugins/Online/OnlineFramework/Source/Party/Private/Party/SocialParty.cpp
@@ -67,7 +67,7 @@ void FPartyRepData::ClearPlatformSessionInfo(const FString& SessionType)
 
 bool FPartyRepData::CanEditData() const
 {
-	return OwnerParty.IsValid() && OwnerParty->IsLocalPlayerPartyLeader();
+	return bAllowOwnerless || (OwnerParty.IsValid() && OwnerParty->IsLocalPlayerPartyLeader());
 }
 
 void FPartyRepData::CompareAgainst(const FOnlinePartyRepDataBase& OldData) const
@@ -1606,7 +1606,7 @@ void USocialParty::RejectAllPendingJoinRequests()
 	{
 		PendingApprovals.Dequeue(PendingApproval);
 		const FUniqueNetId& PrimaryJoiningUserId = *PendingApproval.Members[0].MemberId.GetUniqueNetId();
-		UE_LOG(LogParty, Log, TEXT("[%s] Responding to approval request for %s with denied"), *PartyId.ToString(), *PrimaryJoiningUserId.ToDebugString());
+		UE_LOG(LogParty, Log, TEXT("[%s] Responding to approval request for %s with %s"), *PartyId.ToString(), *PrimaryJoiningUserId.ToDebugString(), ToString(EPartyJoinDenialReason::Busy));
 		if (PendingApproval.bIsJIPApproval)
 		{
 			PRAGMA_DISABLE_DEPRECATION_WARNINGS
diff --git a/Engine/Plugins/Online/OnlineFramework/Source/Party/Private/SocialDebugTools.cpp b/Engine/Plugins/Online/OnlineFramework/Source/Party/Private/SocialDebugTools.cpp
--- a/Engine/Plugins/Online/OnlineFramework/Source/Party/Private/SocialDebugTools.cpp
+++ b/Engine/Plugins/Online/OnlineFramework/Source/Party/Private/SocialDebugTools.cpp
@@ -2,11 +2,13 @@
 
 #include "SocialDebugTools.h"
 
-#include "OnlineSubsystemUtils.h"
 #include "Interfaces/OnlineUserInterface.h"
 #include "Interfaces/OnlinePresenceInterface.h"
+#include "Misc/ScopeExit.h"
+#include "OnlineSubsystemUtils.h"
 #include "Party/PartyMember.h"
 #include "SocialManager.h"
+#include "TimerManager.h"
 
 #include UE_INLINE_GENERATED_CPP_BY_NAME(SocialDebugTools)
 
@@ -205,6 +207,119 @@ void USocialDebugTools::Logout(const FString& Instance, const FLogoutComplete& O
 	OnComplete.ExecuteIfBound(bResult);
 }
 
+void USocialDebugTools::JoinInProgress(const FString& Instance, const FJoinInProgressComplete& OnComplete)
+{
+	// Assume an error, and trigger delegate on scope exit
+	TOptional<FJoinInProgressComplete> OnError = OnComplete;
+	ON_SCOPE_EXIT
+	{
+		// Execute error 
+		if (OnError)
+		{
+			// Error reason doesn't matter, but we're very unlikely to receive this from the JIP target
+			OnError->ExecuteIfBound(EPartyJoinDenialReason::OssUnavailable);
+		}
+	};
+
+	FInstanceContext& Context = GetContext(Instance);
+	IOnlineSubsystem* OnlineSub = Context.GetOSS();
+	if (!OnlineSub)
+	{
+		return;
+	}
+	FUniqueNetIdPtr LocalUserId = Context.GetLocalUserId();
+	if (!LocalUserId)
+	{
+		return;
+	}
+
+	// Find the party we're in
+	IOnlinePartyPtr PartySystem = OnlineSub->GetPartyInterface();
+	if (!PartySystem)
+	{
+		return;
+	}
+	FOnlinePartyConstPtr Party = PartySystem->GetParty(*LocalUserId, IOnlinePartySystem::GetPrimaryPartyTypeId());
+	if (!Party)
+	{
+		return;
+	}
+
+	// JIP on the party leader
+	FUniqueNetIdPtr JipTargetId = Party->LeaderId;
+
+	if (!JipTargetId.IsValid())
+	{
+		return;
+	}
+
+	FPartyMemberJoinInProgressRequest Request;
+	Request.Target = JipTargetId;
+	Request.Time = FDateTime::UtcNow().ToUnixTimestamp();
+
+	if (Context.SetJIPRequest(Request))
+	{
+		OnError.Reset(); // No longer assuming an error
+		
+		// Wait for player to respond 
+		TSharedRef<FDelegateHandle> DelegateHandleRef = MakeShared<FDelegateHandle>();
+		TSharedRef<FTimerHandle> TimerHandleRef = MakeShared<FTimerHandle>();
+		
+		auto CleanupJipRequest = [this, Instance, PartySystem = PartySystem.ToWeakPtr(), OnComplete, DelegateHandleRef, TimerHandleRef](EPartyJoinDenialReason Result)
+		{
+			// Clear our JIP request
+			FInstanceContext& LocalContext = GetContext(Instance);
+			LocalContext.SetJIPRequest(FPartyMemberJoinInProgressRequest());
+
+			OnComplete.ExecuteIfBound(Result);
+
+			// Copy lambda captures as they will be deleted through the unbinds here
+			TSharedRef<FDelegateHandle> DelegateHandleCopy = DelegateHandleRef;
+			TSharedRef<FTimerHandle> TimerHandleCopy = TimerHandleRef;
+			IOnlinePartyPtr StrongPartySystem = PartySystem.Pin();
+			UWorld* World = GetWorld();
+			if (StrongPartySystem)
+			{
+				StrongPartySystem->ClearOnPartyMemberDataReceivedDelegate_Handle(*DelegateHandleCopy);
+			}
+			if (World)
+			{
+				World->GetTimerManager().ClearTimer(*TimerHandleCopy);
+			}
+		};
+
+		auto OnPartyMemberDataReceivedLambda = [this, LocalUserId, JipTargetId, Request, CleanupJipRequest](const FUniqueNetId& InLocalUserId, const FOnlinePartyId& InPartyId, const FUniqueNetId& InMemberId, const FName& InNamespace, const FOnlinePartyData& InPartyMemberData)
+		{
+			if (InMemberId == *JipTargetId)
+			{
+				FPartyMemberRepData ReceivedPartyMemberRepData;
+				if (FVariantDataConverter::VariantMapToUStruct(InPartyMemberData.GetKeyValAttrs(), FPartyMemberRepData::StaticStruct(), &ReceivedPartyMemberRepData, 0, CPF_Transient | CPF_RepSkip))
+				{
+					for (const FPartyMemberJoinInProgressResponse& Response : ReceivedPartyMemberRepData.GetJoinInProgressDataResponses())
+					{
+						if (Response.RequestTime != Request.Time ||
+							Response.Requester != LocalUserId)
+						{
+							// Response was not for us.
+							continue;
+						}
+
+						CleanupJipRequest(static_cast<EPartyJoinDenialReason>(Response.DenialReason));
+						return;
+					}
+				}
+			}
+		};
+
+		*DelegateHandleRef = PartySystem->AddOnPartyMemberDataReceivedDelegate_Handle(FOnPartyMemberDataReceivedDelegate::CreateWeakLambda(this, OnPartyMemberDataReceivedLambda));
+		// wait 60 seconds for a response
+		if (UWorld* World = GetWorld())
+		{
+			World->GetTimerManager().SetTimer(*TimerHandleRef, FTimerDelegate::CreateWeakLambda(this, CleanupJipRequest, static_cast<EPartyJoinDenialReason>(EPartyJoinDenialReason::OssUnavailable)), 60, false);
+		}
+	}
+}
+
 void USocialDebugTools::JoinParty(const FString& Instance, const FString& FriendName, const FJoinPartyComplete& OnComplete)
 {
 	FInstanceContext& Context = GetContext(Instance);
@@ -583,6 +698,17 @@ bool USocialDebugTools::RunCommand(const TCHAR* Cmd, const TArray<FString>& Targ
 		}
 		return true;
 	}
+	else if (FParse::Command(&Cmd, TEXT("JIP")))
+	{
+		for (const FString& TargetInstance : TargetInstances)
+		{
+			JoinInProgress(TargetInstance, FJoinInProgressComplete::CreateLambda([this, TargetInstance](EPartyJoinDenialReason Result)
+			{
+				UE_LOG(LogParty, Display, TEXT("JIP OSS context[%s] %s"), *TargetInstance, ToString(Result));
+			}));
+		}
+		return true;
+	}
 	else if (FParse::Command(&Cmd, TEXT("LEAVEPARTY")))
 	{
 		for (const FString& TargetInstance : TargetInstances)
@@ -698,6 +824,61 @@ FUniqueNetIdPtr USocialDebugTools::FInstanceContext::GetLocalUserId() const
 	return LocalUserId;
 }
 
+void USocialDebugTools::FInstanceContext::ModifyPartyField(const FString& FieldName, const FVariantData& FieldValue)
+{
+	if (!OnlineSub)
+	{
+		return;
+	}
+	IOnlinePartyPtr PartySystem = OnlineSub->GetPartyInterface();
+	if (!PartySystem)
+	{
+		return;
+	}
+	FUniqueNetIdPtr LocalUserId = GetLocalUserId();
+	if (!LocalUserId)
+	{
+		return;
+	}
+	TSharedPtr<const FOnlineParty> Party = PartySystem->GetParty(*LocalUserId, IOnlinePartySystem::GetPrimaryPartyTypeId());
+	if (!Party)
+	{
+		return;
+	}
+
+	// Get current party data and update the JIP field
+	if (!PartyMemberData)
+	{
+		if (FOnlinePartyDataConstPtr ExistingPartyMemberData = PartySystem->GetPartyMemberData(*LocalUserId, *Party->PartyId, *LocalUserId, DefaultPartyDataNamespace))
+		{
+			PartyMemberData = MakeShared<FOnlinePartyData>(*ExistingPartyMemberData);
+		}
+	}
+	if (!PartyMemberData)
+	{
+		PartyMemberData = MakeShared<FOnlinePartyData>();
+	}
+	PartyMemberData->SetAttribute(FieldName, FieldValue);
+	PartySystem->UpdatePartyMemberData(*LocalUserId, *Party->PartyId, DefaultPartyDataNamespace , *PartyMemberData);
+}
+
+bool USocialDebugTools::FInstanceContext::SetJIPRequest(const FPartyMemberJoinInProgressRequest& InRequest)
+{
+	FPartyMemberRepData PartyMemberRepData;
+	PartyMemberRepData.MarkOwnerless();
+	PartyMemberRepData.SetJoinInProgressDataRequest(InRequest);
+
+	FOnlinePartyData OnlinePartyData;
+	const FString JipFieldName(TEXT("JoinInProgressData"));
+	if (FVariantDataConverter::UStructToVariantMap(FPartyMemberRepData::StaticStruct(), &PartyMemberRepData, OnlinePartyData.GetKeyValAttrs(), 0, CPF_Transient | CPF_RepSkip) &&
+		ensure(OnlinePartyData.GetKeyValAttrs().Contains(JipFieldName)))
+	{
+		ModifyPartyField(JipFieldName, OnlinePartyData.GetKeyValAttrs().FindChecked(JipFieldName));
+		return true;
+	}
+	return false;
+}
+
 void USocialDebugTools::HandleFriendInviteReceived(const FUniqueNetId& LocalUserId, const FUniqueNetId& FriendId)
 {
 	if (bAutoAcceptFriendInvites)
diff --git a/Engine/Plugins/Online/OnlineFramework/Source/Party/Public/Party/PartyMember.h b/Engine/Plugins/Online/OnlineFramework/Source/Party/Public/Party/PartyMember.h
--- a/Engine/Plugins/Online/OnlineFramework/Source/Party/Public/Party/PartyMember.h
+++ b/Engine/Plugins/Online/OnlineFramework/Source/Party/Public/Party/PartyMember.h
@@ -123,6 +123,8 @@ struct PARTY_API FPartyMemberRepData : public FOnlinePartyRepDataBase
 public:
 	FPartyMemberRepData() {}
 	void SetOwningMember(const class UPartyMember& InOwnerMember);
+	/** Mark the party data as ownerless. This will bypass any "CanEdit" checks. Useful for using this object in a test context. */
+	void MarkOwnerless();
 
 protected:
 	virtual bool CanEditData() const override;
@@ -132,6 +134,7 @@ protected:
 
 private:
 	TWeakObjectPtr<const UPartyMember> OwnerMember;
+	bool bAllowOwnerless = false;
 
 	/** Platform data fields for party replication */
 	UPROPERTY()
diff --git a/Engine/Plugins/Online/OnlineFramework/Source/Party/Public/Party/SocialParty.h b/Engine/Plugins/Online/OnlineFramework/Source/Party/Public/Party/SocialParty.h
--- a/Engine/Plugins/Online/OnlineFramework/Source/Party/Public/Party/SocialParty.h
+++ b/Engine/Plugins/Online/OnlineFramework/Source/Party/Public/Party/SocialParty.h
@@ -38,6 +38,8 @@ struct PARTY_API FPartyRepData : public FOnlinePartyRepDataBase
 public:
 	FPartyRepData() {}
 	void SetOwningParty(const class USocialParty& InOwnerParty);
+	/** Mark the party data as ownerless. This will bypass any "CanEdit" checks. Useful for using this object in a test context. */
+	void MarkOwnerless();
 
 	const FPartyPlatformSessionInfo* FindSessionInfo(const FString& SessionType) const;
 	const TArray<FPartyPlatformSessionInfo>& GetPlatformSessions() const { return PlatformSessions; }
@@ -52,6 +54,7 @@ protected:
 	virtual const USocialParty* GetOwnerParty() const override;
 
 	TWeakObjectPtr<const USocialParty> OwnerParty;
+	bool bAllowOwnerless = false;
 
 	//@todo DanH Party: Isn't this redundant with the party config itself? Why bother putting it here too when the config replicates to everyone already? #suggested
 	/** The privacy settings for the party */
diff --git a/Engine/Plugins/Online/OnlineFramework/Source/Party/Public/SocialDebugTools.h b/Engine/Plugins/Online/OnlineFramework/Source/Party/Public/SocialDebugTools.h
--- a/Engine/Plugins/Online/OnlineFramework/Source/Party/Public/SocialDebugTools.h
+++ b/Engine/Plugins/Online/OnlineFramework/Source/Party/Public/SocialDebugTools.h
@@ -12,10 +12,12 @@ class USocialManager;
 class IOnlineSubsystem;
 class FOnlineAccountCredentials;
 class IOnlinePartyPendingJoinRequestInfo;
+struct FPartyMemberJoinInProgressRequest;
 typedef TSharedPtr<const class IOnlinePartyJoinInfo> IOnlinePartyJoinInfoConstPtr;
 typedef TSharedPtr<class FOnlinePartyData> FOnlinePartyDataPtr;
 typedef TSharedPtr<const class FOnlinePartyData> FOnlinePartyDataConstPtr;
 using FUniqueNetIdPtr = TSharedPtr<const FUniqueNetId>;
+enum class EPartyJoinDenialReason : uint8;
 
 UCLASS(Within = SocialManager, Config = Game)
 class PARTY_API USocialDebugTools : public UObject, public FExec
@@ -46,6 +48,9 @@ public:
 	DECLARE_DELEGATE_OneParam(FJoinPartyComplete, bool);
 	virtual void JoinParty(const FString& Instance, const FString& FriendName, const FJoinPartyComplete& OnComplete);
 
+	DECLARE_DELEGATE_OneParam(FJoinInProgressComplete, EPartyJoinDenialReason);
+	virtual void JoinInProgress(const FString& Instance, const FJoinInProgressComplete& OnComplete);
+
 	DECLARE_DELEGATE_OneParam(FLeavePartyComplete, bool);
 	virtual void LeaveParty(const FString& Instance, const FLeavePartyComplete& OnComplete);
 
@@ -71,6 +76,9 @@ public:
 		IOnlineSubsystem* GetOSS() const { return OnlineSub; }
 		FOnlinePartyDataPtr GetPartyMemberData() const { return PartyMemberData; }
 		FUniqueNetIdPtr GetLocalUserId() const;
+		void ModifyPartyField(const FString& FieldName, const class FVariantData& FieldValue);
+
+		bool SetJIPRequest(const FPartyMemberJoinInProgressRequest& InRequest);
 
 		FString Name;
 		IOnlineSubsystem* OnlineSub;
