Reduce log spam when querying a social user fails

diff --git a/Engine/Plugins/Online/OnlineFramework/Source/Party/Private/SocialToolkit.cpp b/Engine/Plugins/Online/OnlineFramework/Source/Party/Private/SocialToolkit.cpp
--- a/Engine/Plugins/Online/OnlineFramework/Source/Party/Private/SocialToolkit.cpp
+++ b/Engine/Plugins/Online/OnlineFramework/Source/Party/Private/SocialToolkit.cpp
@@ -494,7 +494,7 @@ void USocialToolkit::NotifySubsystemIdEstablished(USocialUser& SocialUser, ESoci
 	}
 	else
 	{
-		FString LogString = FString::Printf(TEXT("SubsystemId [%s] for user [%s] is already in the UsersBySubsystemId map.\n"), *SubsystemId.ToDebugString(), *SocialUser.GetName());
+		FString LogString = FString::Printf(TEXT("SubsystemId [%s] for user [%s] is already in the UsersBySubsystemIds map.\n"), *SubsystemId.ToDebugString(), *SocialUser.GetName());
 
 		LogString += TEXT("Currently in the map:\n");
 		for (const auto& IdUserPair : UsersBySubsystemIds)
@@ -1175,6 +1175,18 @@ void USocialToolkit::HandleGameDestroyed(const FName SessionName, bool bWasSucce
 void USocialToolkit::HandleUserInvalidated(USocialUser& InvalidUser)
 {
 	AllUsers.Remove(&InvalidUser);
+	// Remove user from id maps
+	if (bRemoveInvalidatedUserFromMaps)
+	{
+		for (const ESocialSubsystem& SocialSubsystemType : InvalidUser.GetRelevantSubsystems())
+		{
+			FUniqueNetIdRepl SubsystemUserId = InvalidUser.GetUserId(SocialSubsystemType);
+			if (SubsystemUserId.IsValid())
+			{
+				UsersBySubsystemIds.Remove(SubsystemUserId);
+			}
+		}
+	}
 	OnSocialUserInvalidated().Broadcast(InvalidUser);
 }
 
diff --git a/Engine/Plugins/Online/OnlineFramework/Source/Party/Private/User/SocialUser.cpp b/Engine/Plugins/Online/OnlineFramework/Source/Party/Private/User/SocialUser.cpp
--- a/Engine/Plugins/Online/OnlineFramework/Source/Party/Private/User/SocialUser.cpp
+++ b/Engine/Plugins/Online/OnlineFramework/Source/Party/Private/User/SocialUser.cpp
@@ -1286,7 +1286,7 @@ void USocialUser::EstablishOssInfo(const TSharedRef<FOnlineBlockedPlayer>& InBlo
 
 	if (InBlockedPlayerInfo != SubsystemInfo.BlockedPlayerInfo)
 	{
-		UE_CLOG(!SubsystemInfo.BlockedPlayerInfo.IsValid(), LogParty, Warning, TEXT("SocialUser [%s] is establishing new blocked player info on [%s], but the existing info is still valid."),
+		UE_CLOG(SubsystemInfo.BlockedPlayerInfo.IsValid(), LogParty, Warning, TEXT("SocialUser [%s] is establishing new blocked player info on [%s], but the existing info is still valid."),
 			*ToDebugString(), ToString(SubsystemType));
 
 		SubsystemInfo.BlockedPlayerInfo = InBlockedPlayerInfo;
@@ -1300,7 +1300,7 @@ void USocialUser::EstablishOssInfo(const TSharedRef<FOnlineRecentPlayer>& InRece
 
 	if (InRecentPlayerInfo != SubsystemInfo.RecentPlayerInfo)
 	{
-		UE_CLOG(!SubsystemInfo.RecentPlayerInfo.IsValid(), LogParty, Warning, TEXT("SocialUser [%s] is establishing new recent player info on [%s], but the existing info is still valid."),
+		UE_CLOG(SubsystemInfo.RecentPlayerInfo.IsValid(), LogParty, Warning, TEXT("SocialUser [%s] is establishing new recent player info on [%s], but the existing info is still valid."),
 			*ToDebugString(), ToString(SubsystemType));
 
 		SubsystemInfo.RecentPlayerInfo = InRecentPlayerInfo;
diff --git a/Engine/Plugins/Online/OnlineFramework/Source/Party/Public/SocialToolkit.h b/Engine/Plugins/Online/OnlineFramework/Source/Party/Public/SocialToolkit.h
--- a/Engine/Plugins/Online/OnlineFramework/Source/Party/Public/SocialToolkit.h
+++ b/Engine/Plugins/Online/OnlineFramework/Source/Party/Public/SocialToolkit.h
@@ -41,7 +41,7 @@ namespace FriendAcceptFailureReason
 DECLARE_DELEGATE_OneParam(FUserDependentAction, USocialUser&);
 
 /** Represents the full suite of social functionality available to a given LocalPlayer */
-UCLASS(Within = SocialManager)
+UCLASS(Within = SocialManager, Config = Game)
 class PARTY_API USocialToolkit : public UObject, public FExec
 {
 	GENERATED_BODY()
@@ -309,6 +309,9 @@ private:
 	TArray<TObjectPtr<USocialUser>> AllUsers;
 	TMap<FUniqueNetIdRepl, TWeakObjectPtr<USocialUser>> UsersBySubsystemIds;
 
+	UPROPERTY(config)
+	bool bRemoveInvalidatedUserFromMaps = true;
+
 	UPROPERTY()
 	TWeakObjectPtr<ULocalPlayer> LocalPlayerOwner = nullptr;
 
