Add JIP to list of social debug commands
Move PrintExecCommands to protected so it can be overridden by subclasses

diff --git a/Engine/Plugins/Online/OnlineFramework/Source/Party/Private/SocialDebugTools.cpp b/Engine/Plugins/Online/OnlineFramework/Source/Party/Private/SocialDebugTools.cpp
--- a/Engine/Plugins/Online/OnlineFramework/Source/Party/Private/SocialDebugTools.cpp
+++ b/Engine/Plugins/Online/OnlineFramework/Source/Party/Private/SocialDebugTools.cpp
@@ -43,6 +43,7 @@ void USocialDebugTools::PrintExecCommands() const
 	UE_LOG(LogParty, Log, TEXT("LEAVEPARTY"));
 	UE_LOG(LogParty, Log, TEXT("AUTOACCEPTFRIENDINVITES"));
 	UE_LOG(LogParty, Log, TEXT("AUTOACCEPTPARTYINVITES"));
+	UE_LOG(LogParty, Log, TEXT("JIP"));
 	UE_LOG(LogParty, Log, TEXT("HELP"));
 }
 
