Change PLATFORM_NAME_DESKTOP from PC to DESKTOP since this was causing failures where FUserPlatform::IsDesktop() would return false for pc platforms

diff --git a/Engine/Plugins/Online/OnlineFramework/Source/Party/Private/SocialTypes.cpp b/Engine/Plugins/Online/OnlineFramework/Source/Party/Private/SocialTypes.cpp
--- a/Engine/Plugins/Online/OnlineFramework/Source/Party/Private/SocialTypes.cpp
+++ b/Engine/Plugins/Online/OnlineFramework/Source/Party/Private/SocialTypes.cpp
@@ -11,7 +11,7 @@
 //////////////////////////////////////////////////////////////////////////
 
 #define PLATFORM_NAME_MOBILE	TEXT("MOBILE")
-#define PLATFORM_NAME_DESKTOP	TEXT("PC")
+#define PLATFORM_NAME_DESKTOP	TEXT("DESKTOP")
 #define PLATFORM_NAME_CONSOLE	TEXT("CONSOLE")
 
 FUserPlatform::FUserPlatform()
