Flagging SetOwningParty and SetOwningMember as virtual functions

diff --git a/Engine/Plugins/Online/OnlineFramework/Source/Party/Public/Party/PartyMember.h b/Engine/Plugins/Online/OnlineFramework/Source/Party/Public/Party/PartyMember.h
--- a/Engine/Plugins/Online/OnlineFramework/Source/Party/Public/Party/PartyMember.h
+++ b/Engine/Plugins/Online/OnlineFramework/Source/Party/Public/Party/PartyMember.h
@@ -122,7 +122,7 @@ struct PARTY_API FPartyMemberRepData : public FOnlinePartyRepDataBase
 
 public:
 	FPartyMemberRepData() {}
-	void SetOwningMember(const class UPartyMember& InOwnerMember);
+	virtual void SetOwningMember(const class UPartyMember& InOwnerMember);
 	/** Mark the party data as ownerless. This will bypass any "CanEdit" checks. Useful for using this object in a test context. */
 	void MarkOwnerless();
 
diff --git a/Engine/Plugins/Online/OnlineFramework/Source/Party/Public/Party/SocialParty.h b/Engine/Plugins/Online/OnlineFramework/Source/Party/Public/Party/SocialParty.h
--- a/Engine/Plugins/Online/OnlineFramework/Source/Party/Public/Party/SocialParty.h
+++ b/Engine/Plugins/Online/OnlineFramework/Source/Party/Public/Party/SocialParty.h
@@ -37,7 +37,7 @@ struct PARTY_API FPartyRepData : public FOnlinePartyRepDataBase
 
 public:
 	FPartyRepData() {}
-	void SetOwningParty(const class USocialParty& InOwnerParty);
+	virtual void SetOwningParty(const class USocialParty& InOwnerParty);
 	/** Mark the party data as ownerless. This will bypass any "CanEdit" checks. Useful for using this object in a test context. */
 	void MarkOwnerless();
 
