Remove USocialManager::AllManagersByGameInstance
It is a static array that is not used anywhere

diff --git a/Engine/Plugins/Online/OnlineFramework/Source/Party/Private/SocialManager.cpp b/Engine/Plugins/Online/OnlineFramework/Source/Party/Private/SocialManager.cpp
--- a/Engine/Plugins/Online/OnlineFramework/Source/Party/Private/SocialManager.cpp
+++ b/Engine/Plugins/Online/OnlineFramework/Source/Party/Private/SocialManager.cpp
@@ -98,7 +98,6 @@ const FName USocialManager::FJoinPartyAttempt::Step_WaitForPersistentPartyCreati
 
 TArray<ESocialSubsystem> USocialManager::DefaultSubsystems;
 TArray<FSocialInteractionHandle> USocialManager::RegisteredInteractions;
-TMap<TWeakObjectPtr<UGameInstance>, TWeakObjectPtr<USocialManager>> USocialManager::AllManagersByGameInstance;
 
 /*static*/bool USocialManager::IsSocialSubsystemEnabled(ESocialSubsystem SubsystemType)
 {
@@ -161,11 +160,6 @@ USocialManager::USocialManager()
 {
 	if (!IsTemplate())
 	{
-		if (ensureMsgf(!AllManagersByGameInstance.Contains(&GetGameInstance()), TEXT("More than one SocialManager has been created for a game instance! Chaos is sure to ensue. Make sure you only have a single instance living on your GameInstance.")))
-		{
-			AllManagersByGameInstance.Add(&GetGameInstance(), this);
-		}
-
 		if (DefaultSubsystems.Num() == 0)
 		{
 			//@todo DanH social: This module assumes there is a primary (aka mcp) OSS available that other accounts are linked to. Consider whether we want to support platform-only situations with this module #future
diff --git a/Engine/Plugins/Online/OnlineFramework/Source/Party/Public/SocialManager.h b/Engine/Plugins/Online/OnlineFramework/Source/Party/Public/SocialManager.h
--- a/Engine/Plugins/Online/OnlineFramework/Source/Party/Public/SocialManager.h
+++ b/Engine/Plugins/Online/OnlineFramework/Source/Party/Public/SocialManager.h
@@ -266,7 +266,6 @@ protected: // overridable handlers
 private:
 	static TArray<ESocialSubsystem> DefaultSubsystems;
 	static TArray<FSocialInteractionHandle> RegisteredInteractions;
-	static TMap<TWeakObjectPtr<UGameInstance>, TWeakObjectPtr<USocialManager>> AllManagersByGameInstance;
 
 	UPROPERTY()
 	TArray<TObjectPtr<USocialToolkit>> SocialToolkits;
