Allow specificying Getter and OnChanged event access levels in the implementation of the EXPOSE macro

diff --git a/Engine/Plugins/Online/OnlineFramework/Source/Party/Public/Party/PartyTypes.h b/Engine/Plugins/Online/OnlineFramework/Source/Party/Public/Party/PartyTypes.h
--- a/Engine/Plugins/Online/OnlineFramework/Source/Party/Public/Party/PartyTypes.h
+++ b/Engine/Plugins/Online/OnlineFramework/Source/Party/Public/Party/PartyTypes.h
@@ -471,8 +471,7 @@ protected:
 // Boilerplate for exposing RepData properties
 //////////////////////////////////////////////////////////////////////////
 
-/** Simplest option - exposes getter and events, but no default setter */
-#define EXPOSE_REP_DATA_PROPERTY_NO_SETTER(Owner, PropertyType, PropertyName, PropertyAccess)	\
+#define EXPOSE_REP_DATA_PROPERTY_NO_SETTER_IMPL(Owner, PropertyType, PropertyName, PropertyAccess, GetterAccess, OnChangedAccess)	\
 public:	\
 	/** If the property is a POD or ptr type, we'll work with it by copy. Otherwise, by const ref */	\
 	using Mutable##PropertyName##Type = std::remove_const_t<PropertyType>;	\
@@ -482,12 +481,12 @@ private:	\
 	/** Bummer to have two signatures, but cases that want both the old and new values are much rarer, so most don't want to bother with a handler that takes an extra unused param */	\
 	DECLARE_MULTICAST_DELEGATE_OneParam(FOn##PropertyName##Changed, PropertyName##ArgType /*NewValue*/);	\
 	DECLARE_MULTICAST_DELEGATE_TwoParams(FOn##PropertyName##ChangedDif, PropertyName##ArgType /*NewValue*/, PropertyName##ArgType /*OldValue*/);	\
-public:	\
+OnChangedAccess:	\
 	/** Bind to receive the new property value only on changes */	\
 	FOn##PropertyName##Changed& On##PropertyName##Changed() const { return On##PropertyName##ChangedEvent; }	\
 	/** Bind to receive both the new and old property value on changes */	\
 	FOn##PropertyName##ChangedDif& On##PropertyName##ChangedDif() const { return On##PropertyName##ChangedDifEvent; }	\
-	\
+GetterAccess: \
 	PropertyName##ArgType Get##PropertyName() const { return PropertyAccess; }	\
 private:	\
 	void Compare##PropertyName(const Owner& OldData) const	\
@@ -506,6 +505,10 @@ private:	\
 	mutable FOn##PropertyName##Changed On##PropertyName##ChangedEvent;	\
 	mutable FOn##PropertyName##ChangedDif On##PropertyName##ChangedDifEvent
 
+/** Simplest option - exposes getter and events, but no default setter */
+#define EXPOSE_REP_DATA_PROPERTY_NO_SETTER(Owner, PropertyType, PropertyName, PropertyAccess)	\
+	EXPOSE_REP_DATA_PROPERTY_NO_SETTER_IMPL(Owner, PropertyType, PropertyName, PropertyAccess, public, public)
+
 #define EXPOSE_REP_DATA_PROPERTY_SETTER_ONLY(Owner, PropertyType, PropertyName, PropertyAccess, SetterPrivacy)	\
 SetterPrivacy:	\
 	void Set##PropertyName(PropertyName##ArgType New##PropertyName)	\
