Made SocialManager's toolkit accessors to be public.
Made SocialManager's join party attempt to be public.
Added a GetPartyMemberIds utility function to get the member ids for a party
Added ToString for EMemberConnectionStatus

diff --git a/Engine/Plugins/Online/OnlineFramework/Source/Party/Private/Party/SocialParty.cpp b/Engine/Plugins/Online/OnlineFramework/Source/Party/Private/Party/SocialParty.cpp
--- a/Engine/Plugins/Online/OnlineFramework/Source/Party/Private/Party/SocialParty.cpp
+++ b/Engine/Plugins/Online/OnlineFramework/Source/Party/Private/Party/SocialParty.cpp
@@ -1621,21 +1621,21 @@ void USocialParty::PumpApprovalQueue()
 	}
 }
 
-namespace UE::Online::Party::Private
+namespace
 {
-	FPartyJoinDenialReason ReservationResponseToJoinDenialReason(const EPartyReservationResult::Type ReservationResponse)
+FPartyJoinDenialReason ReservationResponseToJoinDenialReason(const EPartyReservationResult::Type ReservationResponse)
+{
+	switch (ReservationResponse)
 	{
-		switch (ReservationResponse)
-		{
-		case EPartyReservationResult::PartyLimitReached:
-			return EPartyJoinDenialReason::GameFull;
-		case EPartyReservationResult::ReservationDenied:
-			return EPartyJoinDenialReason::GameFull; // Not really a more specific reason available
-		case EPartyReservationResult::ReservationDenied_CrossPlayRestriction:
-			return EPartyJoinDenialReason::JoinerCrossplayRestricted;
-		}
-		return EPartyJoinDenialReason::NoReason;
+	case EPartyReservationResult::PartyLimitReached:
+		return EPartyJoinDenialReason::GameFull;
+	case EPartyReservationResult::ReservationDenied:
+		return EPartyJoinDenialReason::GameFull; // Not really a more specific reason available
+	case EPartyReservationResult::ReservationDenied_CrossPlayRestriction:
+		return EPartyJoinDenialReason::JoinerCrossplayRestricted;
 	}
+	return EPartyJoinDenialReason::NoReason;
+}
 }
 
 void USocialParty::HandleReservationRequestComplete(EPartyReservationResult::Type ReservationResponse)
@@ -1643,7 +1643,7 @@ void USocialParty::HandleReservationRequestComplete(EPartyReservationResult::Typ
 	UE_LOG(LogParty, Verbose, TEXT("Reservation request complete with response: %s"), EPartyReservationResult::ToString(ReservationResponse));
 
 	const bool bReservationApproved = ReservationResponse == EPartyReservationResult::ReservationAccepted || ReservationResponse == EPartyReservationResult::ReservationDuplicate;
-	const FPartyJoinDenialReason DenialReason = UE::Online::Party::Private::ReservationResponseToJoinDenialReason(ReservationResponse);
+	const FPartyJoinDenialReason DenialReason = ReservationResponseToJoinDenialReason(ReservationResponse);
 
 	if (bReservationApproved || DenialReason.HasAnyReason())
 	{
@@ -2073,3 +2073,26 @@ void USocialParty::RunJoinInProgressTimer()
 		GetWorld()->GetTimerManager().SetTimer(JoinInProgressTimerHandle, this, &USocialParty::RunJoinInProgressTimer, static_cast<float>(NextTimer));
 	}
 }
+
+namespace UE::OnlineFramework::Party
+{
+TArray<FUniqueNetIdRepl> GetPartyMemberIds(const USocialParty* SocialParty)
+{
+	TArray<FUniqueNetIdRepl> PartyMemberIds;
+	if (SocialParty)
+	{
+		auto IsPartyMemberValid = [](const UPartyMember* PartyMember)
+		{
+			CA_ASSUME(PartyMember); // GetPartyMembers filters null members
+			return PartyMember->GetPrimaryNetId().IsValid();
+		};
+		auto GetPartyMemberId = [](const UPartyMember* PartyMember)
+		{
+			CA_ASSUME(PartyMember); // GetPartyMembers filters null members
+			return PartyMember->GetPrimaryNetId();
+		};
+		Algo::TransformIf(SocialParty->GetPartyMembers(), PartyMemberIds, IsPartyMemberValid, GetPartyMemberId);
+	}
+	return PartyMemberIds;
+}
+} // UE::OnlineFramework::Party
diff --git a/Engine/Plugins/Online/OnlineFramework/Source/Party/Private/SocialManager.cpp b/Engine/Plugins/Online/OnlineFramework/Source/Party/Private/SocialManager.cpp
--- a/Engine/Plugins/Online/OnlineFramework/Source/Party/Private/SocialManager.cpp
+++ b/Engine/Plugins/Online/OnlineFramework/Source/Party/Private/SocialManager.cpp
@@ -63,10 +63,6 @@ USocialManager::FJoinPartyAttempt::FJoinPartyAttempt(const USocialUser* InTarget
 	, OnJoinComplete(InOnJoinComplete)
 {}
 
-USocialManager::FJoinPartyAttempt::FJoinPartyAttempt(const USocialUser* InTargetUser, const FOnlinePartyTypeId& InPartyTypeId, const FOnJoinPartyAttemptComplete& InOnJoinComplete)
-	: USocialManager::FJoinPartyAttempt::FJoinPartyAttempt(InTargetUser, InPartyTypeId, PartyJoinMethod::Unspecified, InOnJoinComplete)
-{}
-
 USocialManager::FJoinPartyAttempt::FJoinPartyAttempt(TSharedRef<const FRejoinableParty> InRejoinInfo)
 	: PartyTypeId(IOnlinePartySystem::GetPrimaryPartyTypeId())
 	, JoinMethod(InRejoinInfo->OriginalJoinMethod)
@@ -75,13 +71,12 @@ USocialManager::FJoinPartyAttempt::FJoinPartyAttempt(TSharedRef<const FRejoinabl
 
 FString USocialManager::FJoinPartyAttempt::ToDebugString() const
 {
-	return FString::Printf(TEXT("IsRejoin (%s), TargetUser (%s), PartyId (%s), TypeId (%d), TargetUserPlatformId (%s), PlatformSessionId (%s), JoinMethod (%s)"),
+	return FString::Printf(TEXT("IsRejoin (%s), TargetUser (%s), PartyId (%s), TypeId (%d), TargetUserPlatformId (%s), JoinMethod (%s)"),
 		RejoinInfo.IsValid() ? TEXT("true") : TEXT("false"),
 		TargetUser.IsValid() ? *TargetUser->ToDebugString() : TEXT("invalid"),
 		JoinInfo.IsValid() ? *JoinInfo->GetPartyId()->ToDebugString() : RejoinInfo.IsValid() ? *RejoinInfo->PartyId->ToDebugString() : TEXT("unknown"),
 		PartyTypeId.GetValue(),
 		*TargetUserPlatformId.ToDebugString(),
-		*PlatformSessionId,
 		*JoinMethod.ToString());
 }
 
diff --git a/Engine/Plugins/Online/OnlineFramework/Source/Party/Public/Party/PartyMember.h b/Engine/Plugins/Online/OnlineFramework/Source/Party/Public/Party/PartyMember.h
--- a/Engine/Plugins/Online/OnlineFramework/Source/Party/Public/Party/PartyMember.h
+++ b/Engine/Plugins/Online/OnlineFramework/Source/Party/Public/Party/PartyMember.h
@@ -225,7 +225,6 @@ protected:
 	virtual void Shutdown();
 
 	FPartyMemberDataReplicator MemberDataReplicator;
-	void SetMemberConnectionStatus(EMemberConnectionStatus InMemberConnectionStatus);
 
 	TSharedPtr<const FOnlinePartyMember> GetOSSPartyMember() const { return OssPartyMember; }
 
diff --git a/Engine/Plugins/Online/OnlineFramework/Source/Party/Public/Party/SocialParty.h b/Engine/Plugins/Online/OnlineFramework/Source/Party/Public/Party/SocialParty.h
--- a/Engine/Plugins/Online/OnlineFramework/Source/Party/Public/Party/SocialParty.h
+++ b/Engine/Plugins/Online/OnlineFramework/Source/Party/Public/Party/SocialParty.h
@@ -497,6 +497,11 @@ private:
 	mutable FOnPartyJIPResponse OnPartyJIPResponseEvent;
 };
 
+namespace UE::OnlineFramework::Party
+{
+PARTY_API TArray<FUniqueNetIdRepl> GetPartyMemberIds(const USocialParty* SocialParty);
+} // UE::OnlineFramework::Party
+
 #if UE_ENABLE_INCLUDE_ORDER_DEPRECATED_IN_5_2
 #include "PartyBeaconState.h"
 #endif
diff --git a/Engine/Plugins/Online/OnlineFramework/Source/Party/Public/SocialManager.h b/Engine/Plugins/Online/OnlineFramework/Source/Party/Public/SocialManager.h
--- a/Engine/Plugins/Online/OnlineFramework/Source/Party/Public/SocialManager.h
+++ b/Engine/Plugins/Online/OnlineFramework/Source/Party/Public/SocialManager.h
@@ -35,7 +35,6 @@ class PARTY_API USocialManager : public UObject, public FExec
 	GENERATED_BODY()
 
 	friend class FPartyPlatformSessionManager;
-	friend UPartyMember;
 	friend USocialUser;
 
 public:
@@ -60,6 +59,8 @@ public:
 	virtual void ShutdownSocialManager();
 
 	USocialToolkit& GetSocialToolkit(const ULocalPlayer& LocalPlayer) const;
+	USocialToolkit* GetSocialToolkit(int32 LocalPlayerNum) const;
+	USocialToolkit* GetSocialToolkit(FUniqueNetIdRepl LocalUserId) const;
 	USocialToolkit* GetFirstLocalUserToolkit() const;
 	FUniqueNetIdRepl GetFirstLocalUserId(ESocialSubsystem SubsystemType) const;
 	bool IsLocalUser(const FUniqueNetIdRepl& LocalUserId, ESocialSubsystem SubsystemType) const;
@@ -121,9 +122,6 @@ protected:
 	DECLARE_DELEGATE_OneParam(FOnJoinPartyAttemptComplete, const FJoinPartyResult&);
 	void JoinParty(const USocialUser& UserToJoin, const FOnlinePartyTypeId& PartyTypeId, const FOnJoinPartyAttemptComplete& OnJoinPartyComplete, const FName& JoinMethod);
 
-	USocialToolkit* GetSocialToolkit(int32 LocalPlayerNum) const;
-	USocialToolkit* GetSocialToolkit(FUniqueNetIdRepl LocalUserId) const;
-
 protected:
 	struct PARTY_API FRejoinableParty : public TSharedFromThis<FRejoinableParty>
 	{
@@ -134,24 +132,21 @@ protected:
 		FName OriginalJoinMethod;
 	};
 
+public:
 	struct PARTY_API FJoinPartyAttempt
 	{
 		FJoinPartyAttempt(TSharedRef<const FRejoinableParty> InRejoinInfo);
 		FJoinPartyAttempt(const USocialUser* InTargetUser, const FOnlinePartyTypeId& InPartyTypeId, const FName& InJoinMethod, const FOnJoinPartyAttemptComplete& InOnJoinComplete);
 
-		UE_DEPRECATED(5.1, "This constructor is deprecated, use (USocialUser*, FOnlinePartyTypeId, FName, FOnJoinPartyAttemptComplete) instead.")
-		FJoinPartyAttempt(const USocialUser* InTargetUser, const FOnlinePartyTypeId& InPartyTypeId, const FOnJoinPartyAttemptComplete& InOnJoinComplete);
-
 		FString ToDebugString() const;
 
 		TWeakObjectPtr<const USocialUser> TargetUser;
 		FOnlinePartyTypeId PartyTypeId;
 		FName JoinMethod = PartyJoinMethod::Unspecified;
 		FUniqueNetIdRepl TargetUserPlatformId;
-		FSessionId PlatformSessionId;
 
 		TSharedPtr<const FRejoinableParty> RejoinInfo;
-		TSharedPtr<const IOnlinePartyJoinInfo> JoinInfo;
+		IOnlinePartyJoinInfoConstPtr JoinInfo;
 
 		FOnJoinPartyAttemptComplete OnJoinComplete;
 
@@ -164,6 +159,7 @@ protected:
 
 		FSocialActionTimeTracker ActionTimeTracker;
 	};
+protected:
 
 	virtual void RegisterSocialInteractions();
 
diff --git a/Engine/Plugins/Online/OnlineSubsystem/Source/Private/OnlinePartyInterface.cpp b/Engine/Plugins/Online/OnlineSubsystem/Source/Private/OnlinePartyInterface.cpp
--- a/Engine/Plugins/Online/OnlineSubsystem/Source/Private/OnlinePartyInterface.cpp
+++ b/Engine/Plugins/Online/OnlineSubsystem/Source/Private/OnlinePartyInterface.cpp
@@ -220,6 +220,23 @@ EPartyState EPartyStateFromString(const TCHAR* Value)
 	return EPartyState::None;
 }
 
+const TCHAR* ToString(const EMemberConnectionStatus Value)
+{
+	switch (Value)
+	{
+	case EMemberConnectionStatus::Uninitialized:
+		return TEXT("Uninitialized");
+	case EMemberConnectionStatus::Disconnected:
+		return TEXT("Disconnected");
+	case EMemberConnectionStatus::Initializing:
+		return TEXT("Initializing");
+	case EMemberConnectionStatus::Connected:
+		return TEXT("Connected");
+	}
+	checkNoEntry();
+	return TEXT("Uninitialized");
+}
+
 const TCHAR* ToString(const EMemberExitedReason Value)
 {
 	switch (Value)
diff --git a/Engine/Plugins/Online/OnlineSubsystem/Source/Public/Interfaces/OnlinePartyInterface.h b/Engine/Plugins/Online/OnlineSubsystem/Source/Public/Interfaces/OnlinePartyInterface.h
--- a/Engine/Plugins/Online/OnlineSubsystem/Source/Public/Interfaces/OnlinePartyInterface.h
+++ b/Engine/Plugins/Online/OnlineSubsystem/Source/Public/Interfaces/OnlinePartyInterface.h
@@ -2071,6 +2071,8 @@ ONLINESUBSYSTEM_API const TCHAR* ToString(const EPartyState Value);
 /** @return the enum version of the string passed in */
 ONLINESUBSYSTEM_API EPartyState EPartyStateFromString(const TCHAR* Value);
 
+/** @return the stringified version of the enum passed in */
+ONLINESUBSYSTEM_API const TCHAR* ToString(const EMemberConnectionStatus Value);
 /** @return the stringified version of the enum passed in */
 ONLINESUBSYSTEM_API const TCHAR* ToString(const EMemberExitedReason Value);
 /** @return the stringified version of the enum passed in */
