Downgrade Error to Warning

diff --git a/Engine/Plugins/Online/OnlineFramework/Source/Hotfix/Private/OnlineHotfixManager.cpp b/Engine/Plugins/Online/OnlineFramework/Source/Hotfix/Private/OnlineHotfixManager.cpp
--- a/Engine/Plugins/Online/OnlineFramework/Source/Hotfix/Private/OnlineHotfixManager.cpp
+++ b/Engine/Plugins/Online/OnlineFramework/Source/Hotfix/Private/OnlineHotfixManager.cpp
@@ -453,7 +453,7 @@ void UOnlineHotfixManager::OnEnumerateFilesComplete(bool bWasSuccessful, const F
 	}
 	else
 	{
-		UE_LOG(LogHotfixManager, Error, TEXT("Enumeration of hotfix files failed"));
+		UE_LOG(LogHotfixManager, Warning, TEXT("Enumeration of hotfix files failed"));
 		TriggerHotfixComplete(EHotfixResult::Failed);
 	}
 }
@@ -548,7 +548,7 @@ void UOnlineHotfixManager::OnEnumerateFilesForAvailabilityComplete(bool bWasSucc
 	}
 	else
 	{
-		UE_LOG(LogHotfixManager, Error, TEXT("Enumeration of hotfix files failed"));
+		UE_LOG(LogHotfixManager, Warning, TEXT("Enumeration of hotfix files failed"));
 	}
 
 	OnlineTitleFile = nullptr;
