Removed deprecated join-in-progress code

diff --git a/Engine/Plugins/Online/OnlineFramework/Source/Party/Private/Party/SocialParty.cpp b/Engine/Plugins/Online/OnlineFramework/Source/Party/Private/Party/SocialParty.cpp
--- a/Engine/Plugins/Online/OnlineFramework/Source/Party/Private/Party/SocialParty.cpp
+++ b/Engine/Plugins/Online/OnlineFramework/Source/Party/Private/Party/SocialParty.cpp
@@ -134,26 +134,6 @@ ECrossplayPreference GetCrossplayPreferenceFromJoinData(const FOnlinePartyData&
 	return ECrossplayPreference::NoSelection;
 }
 
-// DEPRECATED - Use the new join in progress flow with USocialParty::RequestJoinInProgress.
-FPartyJoinApproval USocialParty::EvaluateJIPRequest(const FUniqueNetId& PlayerId) const
-{
-	FPartyJoinApproval JoinApproval;
-
-	JoinApproval.SetApprovalAction(EApprovalAction::Deny);
-	JoinApproval.SetDenialReason(EPartyJoinDenialReason::GameModeRestricted);
-	for (const UPartyMember* Member : GetPartyMembers())
-	{
-		// Make sure we are already in the party.
-		if (Member->GetPrimaryNetId() == PlayerId)
-		{
-			JoinApproval.SetApprovalAction(EApprovalAction::EnqueueAndStartBeacon);
-			JoinApproval.SetDenialReason(EPartyJoinDenialReason::NoReason);
-			break;
-		}
-	}
-	return JoinApproval;
-}
-
 bool USocialParty::ApplyCrossplayRestriction(FPartyJoinApproval& JoinApproval, const FUserPlatform& Platform, const FOnlinePartyData& JoinData) const
 {
 	const ECrossplayPreference SenderCrossplayPreference = GetCrossplayPreferenceFromJoinData(JoinData);
diff --git a/Engine/Plugins/Online/OnlineFramework/Source/Party/Private/User/SocialUser.cpp b/Engine/Plugins/Online/OnlineFramework/Source/Party/Private/User/SocialUser.cpp
--- a/Engine/Plugins/Online/OnlineFramework/Source/Party/Private/User/SocialUser.cpp
+++ b/Engine/Plugins/Online/OnlineFramework/Source/Party/Private/User/SocialUser.cpp
@@ -892,11 +892,6 @@ bool USocialUser::CanSendFriendInvite(ESocialSubsystem SubsystemType) const
 	return HasSubsystemInfo(SubsystemType) && !IsLocalUser() && !IsFriend(SubsystemType) && !IsBlocked(SubsystemType) && !IsFriendshipPending(SubsystemType);
 }
 
-void USocialUser::JoinParty(const FOnlinePartyTypeId& PartyTypeId) const
-{
-	JoinParty(PartyTypeId, PartyJoinMethod::Unspecified);
-}
-
 void USocialUser::JoinParty(const FOnlinePartyTypeId& PartyTypeId, const FName& JoinMethod) const
 {
 	const bool bHasSentInvite = HasSentPartyInvite(PartyTypeId);
@@ -1090,11 +1085,6 @@ void USocialUser::HandleRequestToJoinRemoved(const IOnlinePartyRequestToJoinInfo
 	NotifyRequestToJoinRemoved(Request, Reason);
 }
 
-void USocialUser::RequestToJoinParty()
-{
-	RequestToJoinParty(PartyJoinMethod::Unspecified);
-}
-
 void USocialUser::RequestToJoinParty(const FName& JoinMethod)
 {
 	IOnlinePartyPtr PartyInterface = Online::GetPartyInterfaceChecked(GetWorld());
diff --git a/Engine/Plugins/Online/OnlineFramework/Source/Party/Public/Party/SocialParty.h b/Engine/Plugins/Online/OnlineFramework/Source/Party/Public/Party/SocialParty.h
--- a/Engine/Plugins/Online/OnlineFramework/Source/Party/Public/Party/SocialParty.h
+++ b/Engine/Plugins/Online/OnlineFramework/Source/Party/Public/Party/SocialParty.h
@@ -201,14 +201,6 @@ public:
 	DECLARE_EVENT_OneParam(USocialParty, FOnInviteSent, const USocialUser&);
 	FOnInviteSent& OnInviteSent() const { return OnInviteSentEvent; }
 
-	DECLARE_EVENT_TwoParams(USocialParty, FOnPartyJIPApproved, const FOnlinePartyId&, bool /* Success*/);
-	UE_DEPRECATED(5.1, "Use the new join in progress flow with USocialParty::RequestJoinInProgress.")
-	FOnPartyJIPApproved& OnPartyJIPApproved() const { return OnPartyJIPApprovedEvent; }
-
-	DECLARE_EVENT_ThreeParams(USocialParty, FOnPartyJIPResponse, const FOnlinePartyId&, bool /* Success*/, const FString& /*DeniedResultCode*/);
-	UE_DEPRECATED(5.1, "Use the new join in progress flow with USocialParty::RequestJoinInProgress.")
-	FOnPartyJIPResponse& OnPartyJIPResponse() const { return OnPartyJIPResponseEvent; }
-
 	DECLARE_EVENT_TwoParams(USocialParty, FOnPartyMemberConnectionStatusChanged, UPartyMember&, EMemberConnectionStatus);
 	FOnPartyMemberConnectionStatusChanged& OnPartyMemberConnectionStatusChanged() const { return OnPartyMemberConnectionStatusChangedEvent; }
 
@@ -277,10 +269,6 @@ protected:
 	/** Determines the joinability of this party for a group of users requesting to join */
 	virtual FPartyJoinApproval EvaluateJoinRequest(const TArray<IOnlinePartyUserPendingJoinRequestInfoConstRef>& Players, bool bFromJoinRequest) const;
 
-	/** Determines the joinability of the game a party is in for JoinInProgress */
-	UE_DEPRECATED(5.1, "Use the new join in progress flow with USocialParty::RequestJoinInProgress.")
-	virtual FPartyJoinApproval EvaluateJIPRequest(const FUniqueNetId& PlayerId) const;
-
 	/** Determines the reason why, if at all, this party is currently flat-out unjoinable  */
 	virtual FPartyJoinDenialReason DetermineCurrentJoinability() const;
 
@@ -493,8 +481,6 @@ private:
 	mutable FOnPartyMemberConnectionStatusChanged OnPartyMemberConnectionStatusChangedEvent;
 	mutable FOnPartyFunctionalityDegradedChanged OnPartyFunctionalityDegradedChangedEvent;
 	mutable FOnInviteSent OnInviteSentEvent;
-	mutable FOnPartyJIPApproved OnPartyJIPApprovedEvent;
-	mutable FOnPartyJIPResponse OnPartyJIPResponseEvent;
 };
 
 namespace UE::OnlineFramework::Party
diff --git a/Engine/Plugins/Online/OnlineFramework/Source/Party/Public/User/SocialUser.h b/Engine/Plugins/Online/OnlineFramework/Source/Party/Public/User/SocialUser.h
--- a/Engine/Plugins/Online/OnlineFramework/Source/Party/Public/User/SocialUser.h
+++ b/Engine/Plugins/Online/OnlineFramework/Source/Party/Public/User/SocialUser.h
@@ -118,9 +118,6 @@ public:
 	void HandleRequestToJoinReceived(const IOnlinePartyRequestToJoinInfo& Request);
 	void HandleRequestToJoinRemoved(const IOnlinePartyRequestToJoinInfo& Request, EPartyRequestToJoinRemovedReason Reason);
 
-	UE_DEPRECATED(5.1, "Use RequestToJoinParty(const FName&) instead of RequestToJoinParty(void)")
-	void RequestToJoinParty();
-
 	virtual void RequestToJoinParty(const FName& JoinMethod);
 	void AcceptRequestToJoinParty() const;
 	void DismissRequestToJoinParty() const;
@@ -131,9 +128,6 @@ public:
 	bool HasSentPartyInvite(const FOnlinePartyTypeId& PartyTypeId) const;
 	FJoinPartyResult CheckPartyJoinability(const FOnlinePartyTypeId& PartyTypeId) const;
 
-	UE_DEPRECATED(5.1, "Use JoinParty(const FOnlinePartyTypeId&, const FName&) instead of JoinParty(const FOnlinePartyTypeId&)")
-	void JoinParty(const FOnlinePartyTypeId& PartyTypeId) const;
-
 	virtual void JoinParty(const FOnlinePartyTypeId& PartyTypeId, const FName& JoinMethod) const;
 	virtual void RejectPartyInvite(const FOnlinePartyTypeId& PartyTypeId);
 
