Remove module references from UOnlineHotfixManager::HotfixIniFile

diff --git a/Engine/Plugins/Online/OnlineFramework/Source/Hotfix/Hotfix.Build.cs b/Engine/Plugins/Online/OnlineFramework/Source/Hotfix/Hotfix.Build.cs
--- a/Engine/Plugins/Online/OnlineFramework/Source/Hotfix/Hotfix.Build.cs
+++ b/Engine/Plugins/Online/OnlineFramework/Source/Hotfix/Hotfix.Build.cs
@@ -14,7 +14,6 @@ public class Hotfix : ModuleRules
 				"Core",
 				"CoreUObject",
 				"Engine",
-                "HTTP",
 				"MoviePlayerProxy",
 				"OnlineSubsystemUtils"
 			}
diff --git a/Engine/Plugins/Online/OnlineFramework/Source/Hotfix/Private/OnlineHotfixManager.cpp b/Engine/Plugins/Online/OnlineFramework/Source/Hotfix/Private/OnlineHotfixManager.cpp
--- a/Engine/Plugins/Online/OnlineFramework/Source/Hotfix/Private/OnlineHotfixManager.cpp
+++ b/Engine/Plugins/Online/OnlineFramework/Source/Hotfix/Private/OnlineHotfixManager.cpp
@@ -1,15 +1,11 @@
 // Copyright Epic Games, Inc. All Rights Reserved.
 
 #include "OnlineHotfixManager.h"
-#include "HttpModule.h"
 #include "Online.h"
 #include "OnlineSubsystemUtils.h"
 #include "UObject/UObjectIterator.h"
 #include "UObject/Package.h"
 
-#include "Logging/LogSuppressionInterface.h"
-
-
 #include "Misc/PackageName.h"
 #include "Misc/EngineVersion.h"
 #include "Misc/FileHelper.h"
@@ -919,8 +915,6 @@ FConfigFile* UOnlineHotfixManager::GetConfigFile(const FString& IniName)
 
 bool UOnlineHotfixManager::HotfixIniFile(const FString& FileName, const FString& IniData)
 {
-	const bool bIsEngineIni = FileName.Contains(TEXT("Engine.ini"));
-
 	// Flush async loading before modifying GConfig.
 	FlushAsyncLoading();
 
@@ -933,10 +927,6 @@ bool UOnlineHotfixManager::HotfixIniFile(const FString& FileName, const FString&
 	TArray<UObject*> PerObjectConfigObjects;
 	int32 StartIndex = 0;
 	int32 EndIndex = 0;
-	bool bUpdateLogSuppression = false;
-	bool bUpdateConsoleVariables = false;
-	bool bUpdateHttpConfigs = false;
-	TSet<FString> OnlineSubSections;
 	TSet<FString> UpdatedSectionNames;
 	// Find the set of object classes that were affected
 	while (StartIndex >= 0 && StartIndex < IniData.Len() && EndIndex >= StartIndex)
@@ -990,32 +980,6 @@ bool UOnlineHotfixManager::HotfixIniFile(const FString& FileName, const FString&
 					break;
 				}
 
-				if (bIsEngineIni)
-				{
-					// TODO replace all of this with bindees to FCoreDelegates::TSOnConfigSectionsChanged()
-					const TCHAR* LogConfigSection = TEXT("[Core.Log]");
-					const TCHAR* ConsoleVariableSection = TEXT("[ConsoleVariables]");
-					const TCHAR* HttpSection = TEXT("[HTTP"); // note "]" omitted on purpose since we want a partial match
-					const TCHAR* OnlineSubSectionKey = TEXT("[OnlineSubsystem"); // note "]" omitted on purpose since we want a partial match
-					if (!bUpdateLogSuppression && FCString::Strnicmp(*IniData + StartIndex, LogConfigSection, FCString::Strlen(LogConfigSection)) == 0)
-					{
-						bUpdateLogSuppression = true;
-					}
-					else if (!bUpdateConsoleVariables && FCString::Strnicmp(*IniData + StartIndex, ConsoleVariableSection, FCString::Strlen(ConsoleVariableSection)) == 0)
-					{
-						bUpdateConsoleVariables = true;
-					}
-					else if (!bUpdateHttpConfigs &&	FCString::Strnicmp(*IniData + StartIndex, HttpSection, FCString::Strlen(HttpSection)) == 0)
-					{
-						bUpdateHttpConfigs = true;
-					}
-					else if (FCString::Strnicmp(*IniData + StartIndex, OnlineSubSectionKey, FCString::Strlen(OnlineSubSectionKey)) == 0)
-					{
-						FString SectionStr = IniData.Mid(StartIndex, EndIndex - StartIndex + 1);
-						OnlineSubSections.Emplace(MoveTemp(SectionStr));
-					}
-				}
-
 				// Per object config entries will have a space in the name, but classes won't
 				if (PerObjectNameIndex == -1 || PerObjectNameIndex > EndIndex)
 				{
@@ -1114,31 +1078,6 @@ bool UOnlineHotfixManager::HotfixIniFile(const FString& FileName, const FString&
 	const FString ConfigFileName = ConfigFile->Name.ToString();
 	FCoreDelegates::TSOnConfigSectionsChanged().Broadcast(ConfigFileName, UpdatedSectionNames);
 
-	// Reload log suppression if configs changed
-	if (bUpdateLogSuppression)
-	{
-		FLogSuppressionInterface::Get().ProcessConfigAndCommandLine();
-	}
-
-	// Reload console variables if configs changed
-	if (bUpdateConsoleVariables)
-	{
-		FConfigCacheIni::LoadConsoleVariablesFromINI();
-	}
-
-	// Reload configs relevant to the HTTP module
-	if (bUpdateHttpConfigs)
-	{
-		FHttpModule::Get().UpdateConfigs();
-	}
-
-	// Reload configs relevant to OSS config sections that were updated
-	IOnlineSubsystem* OnlineSub = IOnlineSubsystem::Get(OSSName.Len() ? FName(*OSSName, FNAME_Find) : NAME_None);
-	if (OnlineSub != nullptr)
-	{
-		OnlineSub->ReloadConfigs(OnlineSubSections);
-	}
-
 	UE_LOG(LogHotfixManager, Log, TEXT("Updating config from %s took %f seconds and reloaded %d objects"),
 		*FileName, FPlatformTime::Seconds() - StartTime, NumObjectsReloaded);
 	return true;
diff --git a/Engine/Plugins/Online/OnlineSubsystem/Source/Private/OnlineSubsystemImpl.cpp b/Engine/Plugins/Online/OnlineSubsystem/Source/Private/OnlineSubsystemImpl.cpp
--- a/Engine/Plugins/Online/OnlineSubsystem/Source/Private/OnlineSubsystemImpl.cpp
+++ b/Engine/Plugins/Online/OnlineSubsystem/Source/Private/OnlineSubsystemImpl.cpp
@@ -2,6 +2,7 @@
 
 #include "OnlineSubsystemImpl.h"
 #include "Misc/App.h"
+#include "Misc/CoreDelegates.h"
 #include "NamedInterfaces.h"
 #include "OnlineError.h"
 #include "Interfaces/OnlineIdentityInterface.h"
@@ -44,6 +45,12 @@ FOnlineSubsystemImpl::FOnlineSubsystemImpl(FName InSubsystemName, FName InInstan
 	NamedInterfaces(nullptr),
 	bTickerStarted(true)
 {
+	FCoreDelegates::OnConfigSectionsChanged.AddRaw(this, &FOnlineSubsystemImpl::OnConfigSectionsChanged);
+}
+
+FOnlineSubsystemImpl::~FOnlineSubsystemImpl()
+{
+	FCoreDelegates::OnConfigSectionsChanged.RemoveAll(this);
 }
 
 void FOnlineSubsystemImpl::PreUnload()
@@ -369,6 +376,21 @@ void FOnlineSubsystemImpl::FinalizeReceipts(const FUniqueNetId& UserId)
 	}
 }
 
+void FOnlineSubsystemImpl::OnConfigSectionsChanged(const FString& IniFilename, const TSet<FString>& SectionNames)
+{
+	if (IniFilename == GEngineIni)
+	{
+		for (const FString& SectionName : SectionNames)
+		{
+			if (SectionName.StartsWith(TEXT("OnlineSubsystem")))
+			{
+				ReloadConfigs(SectionNames);
+				break;
+			}
+		}
+	}
+}
+
 void FOnlineSubsystemImpl::OnQueryReceiptsComplete(const FOnlineError& Result, FUniqueNetIdPtr UserId)
 {
 	UE_LOG_ONLINE(Display, TEXT("OnQueryReceiptsComplete %s"), *Result.ToLogString());
diff --git a/Engine/Plugins/Online/OnlineSubsystem/Source/Public/OnlineSubsystemImpl.h b/Engine/Plugins/Online/OnlineSubsystem/Source/Public/OnlineSubsystemImpl.h
--- a/Engine/Plugins/Online/OnlineSubsystem/Source/Public/OnlineSubsystemImpl.h
+++ b/Engine/Plugins/Online/OnlineSubsystem/Source/Public/OnlineSubsystemImpl.h
@@ -84,7 +84,7 @@ protected:
 
 public:
 	
-	virtual ~FOnlineSubsystemImpl() = default;
+	virtual ~FOnlineSubsystemImpl();
 
 	// IOnlineSubsystem
 	virtual IOnlineGroupsPtr GetGroupsInterface() const override;
@@ -158,6 +158,11 @@ public:
 		ExecuteDelegateNextTick(FNextTickDelegate::CreateLambda(Forward<LAMBDA_TYPE>(Callback)));
 	}
 
+	/**
+	* Delegate for config file changes.
+	*/
+	virtual void OnConfigSectionsChanged(const FString& IniFilename, const TSet<FString>& SectionNames);
+
 	/** Name given to default OSS instances (disambiguates for PIE) */
 	static const FName DefaultInstanceName;
 };
diff --git a/Engine/Source/Runtime/Core/Private/Logging/LogSuppressionInterface.cpp b/Engine/Source/Runtime/Core/Private/Logging/LogSuppressionInterface.cpp
--- a/Engine/Source/Runtime/Core/Private/Logging/LogSuppressionInterface.cpp
+++ b/Engine/Source/Runtime/Core/Private/Logging/LogSuppressionInterface.cpp
@@ -12,6 +12,7 @@
 #include "Misc/Parse.h"
 #include "Containers/Map.h"
 #include "Misc/CoreMisc.h"
+#include "Misc/CoreDelegates.h"
 #include "Misc/CommandLine.h"
 #include "Misc/ConfigCacheIni.h"
 #include "Misc/OutputDeviceHelper.h"
@@ -482,6 +483,14 @@ class FLogSuppressionImplementation: public FLogSuppressionInterface, private FS
 		}
 	}
 
+	void OnConfigSectionsChanged(const FString& IniFilename, const TSet<FString>& SectionNames)
+	{
+		if (IniFilename == GEngineIni && SectionNames.Contains(TEXT("Core.Log")))
+		{
+			ProcessConfigAndCommandLine();
+		}
+	}
+
 	virtual void ProcessConfigAndCommandLine()
 	{
 		ReverseAssociations.Reserve(PendingAssociations.Num());
@@ -554,6 +563,21 @@ class FLogSuppressionImplementation: public FLogSuppressionInterface, private FS
 			SetupSuppress(It.Value(), It.Key());
 		}
 
+		if (!bInitialized)
+		{
+			FCoreDelegates::OnConfigSectionsChanged.AddLambda([this](const FString& IniFilename, const TSet<FString>& SectionNames)
+			{
+				if (auto* LogSuppressionImplementation = static_cast<FLogSuppressionImplementation*>(FLogSuppressionInterface::TryGet()))
+				{
+					LogSuppressionImplementation->OnConfigSectionsChanged(IniFilename, SectionNames);
+				}
+				else
+				{
+					FCoreDelegates::OnConfigSectionsChanged.RemoveAll(this);
+				}
+			});
+		}
+
 		bInitialized = true;
 	}
 
diff --git a/Engine/Source/Runtime/Core/Private/Misc/ConfigCacheIni.cpp b/Engine/Source/Runtime/Core/Private/Misc/ConfigCacheIni.cpp
--- a/Engine/Source/Runtime/Core/Private/Misc/ConfigCacheIni.cpp
+++ b/Engine/Source/Runtime/Core/Private/Misc/ConfigCacheIni.cpp
@@ -2075,6 +2075,18 @@ void FConfigFile::ProcessPropertyAndWriteForDefaults(int IniCombineThreshold, co
 /*-----------------------------------------------------------------------------
 	FConfigCacheIni
 -----------------------------------------------------------------------------*/
+
+namespace
+{
+	void OnConfigSectionsChanged(const FString& IniFilename, const TSet<FString>& SectionNames)
+	{
+		if (IniFilename == GEngineIni && SectionNames.Contains(TEXT("ConsoleVariables")))
+		{
+			FConfigCacheIni::LoadConsoleVariablesFromINI();
+		}
+	}
+}
+
 #if WITH_EDITOR
 static TMap<FName, TFuture<void>>& GetPlatformConfigFutures()
 {
@@ -3845,6 +3857,8 @@ void FConfigCacheIni::InitializeConfigSystem()
 	// load editor, etc config files
 	LoadRemainingConfigFiles(Context);
 
+	FCoreDelegates::OnConfigSectionsChanged.AddStatic(OnConfigSectionsChanged);
+
 	// now we can make use of GConfig
 	GConfig->bIsReadyForUse = true;
 
diff --git a/Engine/Source/Runtime/Online/HTTP/Private/HttpModule.cpp b/Engine/Source/Runtime/Online/HTTP/Private/HttpModule.cpp
--- a/Engine/Source/Runtime/Online/HTTP/Private/HttpModule.cpp
+++ b/Engine/Source/Runtime/Online/HTTP/Private/HttpModule.cpp
@@ -78,6 +78,7 @@ void FHttpModule::StartupModule()
 	HttpThreadIdleMinimumSleepTimeInSeconds = 0.0f;	
 
 	// override the above defaults from configs
+	FCoreDelegates::OnConfigSectionsChanged.AddRaw(this, &FHttpModule::OnConfigSectionsChanged);
 	UpdateConfigs();
 
 	if (!FParse::Value(FCommandLine::Get(), TEXT("httpproxy="), ProxyAddress))
@@ -140,12 +141,29 @@ void FHttpModule::ShutdownModule()
 	// at least on Linux, the code in HTTP manager (e.g. request destructors) expects platform to be initialized yet
 	delete HttpManager;	// can be passed NULLs
 
+	FCoreDelegates::OnConfigSectionsChanged.RemoveAll(this);
+
 	FPlatformHttp::Shutdown();
 
 	HttpManager = nullptr;
 	Singleton = nullptr;
 }
 
+void FHttpModule::OnConfigSectionsChanged(const FString& IniFilename, const TSet<FString>& SectionNames)
+{
+	if (IniFilename == GEngineIni)
+	{
+		for (const FString& SectionName : SectionNames)
+		{
+			if (SectionName.StartsWith(TEXT("HTTP")))
+			{
+				UpdateConfigs();
+				break;
+			}
+		}
+	}
+}
+
 bool FHttpModule::HandleHTTPCommand(const TCHAR* Cmd, FOutputDevice& Ar)
 {
 	if (FParse::Command(&Cmd, TEXT("TEST")))
diff --git a/Engine/Source/Runtime/Online/HTTP/Public/HttpModule.h b/Engine/Source/Runtime/Online/HTTP/Public/HttpModule.h
--- a/Engine/Source/Runtime/Online/HTTP/Public/HttpModule.h
+++ b/Engine/Source/Runtime/Online/HTTP/Public/HttpModule.h
@@ -334,6 +334,10 @@ private:
 	 */
 	virtual void ShutdownModule() override;
 
+	/**
+	 * Delegate for config file changes.
+	 */
+	void OnConfigSectionsChanged(const FString& IniFilename, const TSet<FString>& SectionNames);
 
 	/** Keeps track of Http requests while they are being processed */
 	FHttpManager* HttpManager;
