Add InviteMethod param in CanInvite method of SocialParty and make it public, so it's easier to be extended.

#skipci

diff --git a/Engine/Plugins/Online/OnlineFramework/Source/Party/Private/Party/SocialParty.cpp b/Engine/Plugins/Online/OnlineFramework/Source/Party/Private/Party/SocialParty.cpp
--- a/Engine/Plugins/Online/OnlineFramework/Source/Party/Private/Party/SocialParty.cpp
+++ b/Engine/Plugins/Online/OnlineFramework/Source/Party/Private/Party/SocialParty.cpp
@@ -243,15 +243,35 @@ bool USocialParty::HasUserBeenInvited(const USocialUser& User) const
 	return false;
 }
 
-bool USocialParty::CanInviteUser(const USocialUser& User) const
+bool USocialParty::CanInviteUser(const USocialUser& User, ESocialPartyInviteMethod InviteMethod) const
 {
-	return CanInviteUserInternal(User) == ESocialPartyInviteFailureReason::Success;
+	return CanInviteUserInternal(User, InviteMethod) == ESocialPartyInviteFailureReason::Success;
 }
 
 ESocialPartyInviteFailureReason USocialParty::CanInviteUserInternal(const USocialUser& User) const
 {
-	// Only users that are online can be invited
-	if (!User.IsOnline() && !User.CanReceiveOfflineInvite())
+	UE_LOG(LogParty, Warning, TEXT("Invoking a deprecated method CanInviteUserInternal, use CanInviteUserInternal(const USocialUser& User, const ESocialPartyInviteMethod InviteMethod) instead!"));
+	return CanInviteUserInternal(User, ESocialPartyInviteMethod::Other);
+}
+
+ESocialPartyInviteFailureReason USocialParty::CanInviteUserInternal(const USocialUser& User, ESocialPartyInviteMethod InviteMethod) const
+{
+	bool bCanOnlyInviteFriend = false;
+	switch (InviteMethod)
+	{
+	case ESocialPartyInviteMethod::Notification:
+	case ESocialPartyInviteMethod::Other:
+		bCanOnlyInviteFriend = true;
+		break;
+	case ESocialPartyInviteMethod::AcceptRequestToJoin:
+		break;
+	default:
+		// Can't really decide here for Custom methods, pass the IsOnline check and let the Game itself do the check
+		check(InviteMethod >= ESocialPartyInviteMethod::Custom0 && InviteMethod < ESocialPartyInviteMethod::MAX);
+		break;
+	}
+
+	if (bCanOnlyInviteFriend && !User.IsOnline() && !User.CanReceiveOfflineInvite())
 	{
 		return ESocialPartyInviteFailureReason::NotOnline;
 	}
@@ -303,7 +323,7 @@ bool USocialParty::IsInviteRateLimited(const USocialUser& User, ESocialSubsystem
 bool USocialParty::TryInviteUser(const USocialUser& UserToInvite, const ESocialPartyInviteMethod InviteMethod, const FString& MetaData)
 {
 	bool bSentInvite = false;
-	ESocialPartyInviteFailureReason CanInviteResult = CanInviteUserInternal(UserToInvite);
+	ESocialPartyInviteFailureReason CanInviteResult = CanInviteUserInternal(UserToInvite, InviteMethod);
 	if (CanInviteResult == ESocialPartyInviteFailureReason::Success)
 	{
 		const bool bPreferPlatformInvite = USocialSettings::ShouldPreferPlatformInvites();
diff --git a/Engine/Plugins/Online/OnlineFramework/Source/Party/Private/User/SocialUser.cpp b/Engine/Plugins/Online/OnlineFramework/Source/Party/Private/User/SocialUser.cpp
--- a/Engine/Plugins/Online/OnlineFramework/Source/Party/Private/User/SocialUser.cpp
+++ b/Engine/Plugins/Online/OnlineFramework/Source/Party/Private/User/SocialUser.cpp
@@ -1088,14 +1088,16 @@ void USocialUser::HandleRequestToJoinRemoved(const IOnlinePartyRequestToJoinInfo
 void USocialUser::RequestToJoinParty(const FName& JoinMethod)
 {
 	IOnlinePartyPtr PartyInterface = Online::GetPartyInterfaceChecked(GetWorld());
-	PartyInterface->RequestToJoinParty(*GetOwningToolkit().GetLocalUserNetId(ESocialSubsystem::Primary), IOnlinePartySystem::GetPrimaryPartyTypeId(), *GetUserId(ESocialSubsystem::Primary), FOnRequestToJoinPartyComplete::CreateUObject(this, &USocialUser::HandlePartyRequestToJoinSent, JoinMethod));
+	FPartyInvitationRecipient Recipient(*GetUserId(ESocialSubsystem::Primary));
+	Recipient.PlatformData = JoinMethod.ToString();
+	PartyInterface->RequestToJoinParty(*GetOwningToolkit().GetLocalUserNetId(ESocialSubsystem::Primary), IOnlinePartySystem::GetPrimaryPartyTypeId(), Recipient, FOnRequestToJoinPartyComplete::CreateUObject(this, &USocialUser::HandlePartyRequestToJoinSent, JoinMethod));
 }
 
 void USocialUser::AcceptRequestToJoinParty() const
 {
 	if (USocialParty* Party = GetOwningToolkit().GetSocialManager().GetParty(IOnlinePartySystem::GetPrimaryPartyTypeId()))
 	{
-		Party->TryInviteUser(*this, ESocialPartyInviteMethod::Other, TEXT("RequestToJoin"));
+		Party->TryInviteUser(*this, ESocialPartyInviteMethod::AcceptRequestToJoin, TEXT("RequestToJoin"));
 		IOnlinePartyPtr PartyInterface = Online::GetPartyInterfaceChecked(GetWorld());
 		PartyInterface->ClearRequestToJoinParty(*GetOwningToolkit().GetLocalUserNetId(ESocialSubsystem::Primary), Party->GetPartyId(), *GetUserId(ESocialSubsystem::Primary), EPartyRequestToJoinRemovedReason::Accepted);
 	}
diff --git a/Engine/Plugins/Online/OnlineFramework/Source/Party/Public/Party/PartyTypes.h b/Engine/Plugins/Online/OnlineFramework/Source/Party/Public/Party/PartyTypes.h
--- a/Engine/Plugins/Online/OnlineFramework/Source/Party/Public/Party/PartyTypes.h
+++ b/Engine/Plugins/Online/OnlineFramework/Source/Party/Public/Party/PartyTypes.h
@@ -89,6 +89,10 @@ enum class EPartyJoinDenialReason : uint8
 	MissingPartyClassForTypeId,
 	/** The target user is blocked by the local user on one or more of the active subsystems */
 	TargetUserBlocked,
+	/** Failed to get party join info from target user */
+	InvalidJoinInfo,
+	/** The join method doesn't support non-friends */
+	NotFriends,
 
 	/**
 	* Customizable denial reasons.
@@ -209,6 +213,15 @@ inline const TCHAR* ToString(EPartyJoinDenialReason Type)
 	case EPartyJoinDenialReason::MissingPartyClassForTypeId:
 		return TEXT("MissingPartyClassForTypeId");
 		break;
+	case EPartyJoinDenialReason::TargetUserBlocked:
+		return TEXT("TargetUserBlocked");
+		break;
+	case EPartyJoinDenialReason::InvalidJoinInfo:
+		return TEXT("InvalidJoinInfo");
+		break;
+	case EPartyJoinDenialReason::NotFriends:
+		return TEXT("NotFriends");
+		break;
 	default:
 		return TEXT("CustomReason");
 		break;
@@ -231,8 +244,21 @@ enum class ESocialPartyInviteMethod : uint8
 	Other = 0,
 	/** Invite was sent from a toast */
 	Notification,
-	/** Invite was sent with a custom method */
-	Custom1
+	/** Invite was sent via accepting "Request To Join" */
+	AcceptRequestToJoin,
+	/** Custom methods */
+	Custom0,
+	Custom1,
+	Custom2,
+	Custom3,
+	Custom4,
+	Custom5,
+	Custom6,
+	Custom7,
+	Custom8,
+	Custom9,
+
+	MAX
 };
 
 UENUM()
diff --git a/Engine/Plugins/Online/OnlineFramework/Source/Party/Public/Party/SocialParty.h b/Engine/Plugins/Online/OnlineFramework/Source/Party/Public/Party/SocialParty.h
--- a/Engine/Plugins/Online/OnlineFramework/Source/Party/Public/Party/SocialParty.h
+++ b/Engine/Plugins/Online/OnlineFramework/Source/Party/Public/Party/SocialParty.h
@@ -215,6 +215,8 @@ public:
 	DECLARE_DELEGATE_OneParam(FOnRequestJoinInProgressComplete, const EPartyJoinDenialReason /*DenialReason*/);
 	void RequestJoinInProgress(const UPartyMember& TargetMember, const FOnRequestJoinInProgressComplete& CompletionDelegate);
 
+	bool CanInviteUser(const USocialUser& User, const ESocialPartyInviteMethod InviteMethod = ESocialPartyInviteMethod::Other) const;
+
 protected:
 	void InitializeParty(const TSharedRef<const FOnlineParty>& InOssParty);
 	bool IsInitialized() const;
@@ -231,7 +233,6 @@ protected:
 	// User/member-specific actions that are best exposed on the individuals themselves, but best handled by the actual party
 	bool HasUserBeenInvited(const USocialUser& User) const;
 	
-	bool CanInviteUser(const USocialUser& User) const;
 	bool CanPromoteMember(const UPartyMember& PartyMember) const;
 	bool CanKickMember(const UPartyMember& PartyMember) const;
 	
@@ -255,7 +256,9 @@ protected:
 	virtual void OnLeftPartyInternal(EMemberExitedReason Reason);
 
 	/** Virtual versions of the package-scoped "CanX" methods above, as a virtual declared within package scoping cannot link (exported public, imported protected) */
+	UE_DEPRECATED(5.3, "This function has been deperecated, use CanInviteUserInternal(const USocialUser& User, const ESocialPartyInviteMethod InviteMethod) instead.")
 	virtual ESocialPartyInviteFailureReason CanInviteUserInternal(const USocialUser& User) const;
+	virtual ESocialPartyInviteFailureReason CanInviteUserInternal(const USocialUser& User, const ESocialPartyInviteMethod InviteMethod) const;
 	virtual bool CanPromoteMemberInternal(const UPartyMember& PartyMember) const;
 	virtual bool CanKickMemberInternal(const UPartyMember& PartyMember) const;
 
