Analytics support for FJoinPartyAttempt

#rnx

diff --git a/Engine/Plugins/Online/OnlineFramework/Source/Party/Private/SocialManager.cpp b/Engine/Plugins/Online/OnlineFramework/Source/Party/Private/SocialManager.cpp
--- a/Engine/Plugins/Online/OnlineFramework/Source/Party/Private/SocialManager.cpp
+++ b/Engine/Plugins/Online/OnlineFramework/Source/Party/Private/SocialManager.cpp
@@ -617,6 +617,7 @@ void USocialManager::JoinParty(const USocialUser& UserToJoin, const FOnlineParty
 	UE_LOG(LogParty, Verbose, TEXT("Attempting to join user [%s]'s party of type [%d] by [%s]"), *UserToJoin.ToDebugString(), PartyTypeId.GetValue(), *JoinMethod.ToString());
 
 	FJoinPartyAttempt NewAttempt(&UserToJoin, PartyTypeId, JoinMethod, OnJoinPartyComplete);
+	NewAttempt.AnalyticsContext = UserToJoin.GetAnalyticsContext();
 	const FJoinPartyResult ValidationResult = ValidateJoinTarget(UserToJoin, PartyTypeId);
 	if (ValidationResult.WasSuccessful())
 	{
diff --git a/Engine/Plugins/Online/OnlineFramework/Source/Party/Public/SocialManager.h b/Engine/Plugins/Online/OnlineFramework/Source/Party/Public/SocialManager.h
--- a/Engine/Plugins/Online/OnlineFramework/Source/Party/Public/SocialManager.h
+++ b/Engine/Plugins/Online/OnlineFramework/Source/Party/Public/SocialManager.h
@@ -158,6 +158,8 @@ public:
 		static const FName Step_WaitForPersistentPartyCreation;
 
 		FSocialActionTimeTracker ActionTimeTracker;
+
+		TMap<FString, FString> AnalyticsContext;
 	};
 protected:
 
