Added initialization completed multicast delegate on the social party

#rnx

diff --git a/Engine/Plugins/Online/OnlineFramework/Source/Party/Private/Party/SocialParty.cpp b/Engine/Plugins/Online/OnlineFramework/Source/Party/Private/Party/SocialParty.cpp
--- a/Engine/Plugins/Online/OnlineFramework/Source/Party/Private/Party/SocialParty.cpp
+++ b/Engine/Plugins/Online/OnlineFramework/Source/Party/Private/Party/SocialParty.cpp
@@ -520,6 +520,7 @@ void USocialParty::TryFinishInitialization()
 		if (OSSMemberCount == PartyMembersById.Num() && bHasReceivedRepData)
 		{
 			bIsInitialized = true;
+			OnInitializationCompletePreNotify().Broadcast(*this);
 			GetSocialManager().NotifyPartyInitialized(*this);
 		}
 	}
diff --git a/Engine/Plugins/Online/OnlineFramework/Source/Party/Public/Party/SocialParty.h b/Engine/Plugins/Online/OnlineFramework/Source/Party/Public/Party/SocialParty.h
--- a/Engine/Plugins/Online/OnlineFramework/Source/Party/Public/Party/SocialParty.h
+++ b/Engine/Plugins/Online/OnlineFramework/Source/Party/Public/Party/SocialParty.h
@@ -204,6 +204,9 @@ public:
 	DECLARE_EVENT_TwoParams(USocialParty, FOnPartyMemberConnectionStatusChanged, UPartyMember&, EMemberConnectionStatus);
 	FOnPartyMemberConnectionStatusChanged& OnPartyMemberConnectionStatusChanged() const { return OnPartyMemberConnectionStatusChangedEvent; }
 
+	DECLARE_MULTICAST_DELEGATE_OneParam(FOnInitializationCompletePreNotify, USocialParty&);
+	FOnInitializationCompletePreNotify& OnInitializationCompletePreNotify() const { return OnInitializationCompletePreNotifyEvent; }
+
 	void ResetPrivacySettings();
 	const FPartyPrivacySettings& GetPrivacySettings() const;
 
@@ -484,6 +487,7 @@ private:
 	mutable FOnPartyMemberConnectionStatusChanged OnPartyMemberConnectionStatusChangedEvent;
 	mutable FOnPartyFunctionalityDegradedChanged OnPartyFunctionalityDegradedChangedEvent;
 	mutable FOnInviteSent OnInviteSentEvent;
+	mutable FOnInitializationCompletePreNotify OnInitializationCompletePreNotifyEvent;
 };
 
 namespace UE::OnlineFramework::Party
