CL split for vis.

diff --git a/Engine/Plugins/Online/OnlineFramework/Source/Party/Private/Party/SocialParty.cpp b/Engine/Plugins/Online/OnlineFramework/Source/Party/Private/Party/SocialParty.cpp
--- a/Engine/Plugins/Online/OnlineFramework/Source/Party/Private/Party/SocialParty.cpp
+++ b/Engine/Plugins/Online/OnlineFramework/Source/Party/Private/Party/SocialParty.cpp
@@ -825,6 +825,8 @@ void USocialParty::HandleJoinabilityQueryReceived(const FUniqueNetId& LocalUserI
 
 void USocialParty::HandlePartyDataReceived(const FUniqueNetId& LocalUserId, const FOnlinePartyId& PartyId, const FName& Namespace, const FOnlinePartyData& PartyData)
 {
+	UE_LOG(LogParty, Verbose, TEXT("USocialParty::HandlePartyDataReceived"));
+
 	if (Namespace != DefaultPartyDataNamespace)
 	{
 		return;
