Clear non valid assets referenced by OnlineHotfixManager before garbage collection is triggered to allow freeing assets from unloaded plugins if gc.PendingKillEnabled = false

diff --git a/Engine/Plugins/Online/OnlineFramework/Source/Hotfix/Private/OnlineHotfixManager.cpp b/Engine/Plugins/Online/OnlineFramework/Source/Hotfix/Private/OnlineHotfixManager.cpp
--- a/Engine/Plugins/Online/OnlineFramework/Source/Hotfix/Private/OnlineHotfixManager.cpp
+++ b/Engine/Plugins/Online/OnlineFramework/Source/Hotfix/Private/OnlineHotfixManager.cpp
@@ -163,6 +163,11 @@ UOnlineHotfixManager::UOnlineHotfixManager() :
 	bLogMountedPakContents = FParse::Param(FCommandLine::Get(), TEXT("LogHotfixPakContents"));
 #endif
 	GameContentPath = FString() / FApp::GetProjectName() / TEXT("Content");
+
+	if (!UObject::IsPendingKillEnabled())
+	{
+		FCoreUObjectDelegates::GetPreGarbageCollectDelegate().AddUObject(this, &UOnlineHotfixManager::StopTrackingInvalidHotfixedAssets);
+	}
 }
 
 UOnlineHotfixManager::UOnlineHotfixManager(FVTableHelper& Helper)
@@ -172,6 +177,10 @@ UOnlineHotfixManager::UOnlineHotfixManager(FVTableHelper& Helper)
 
 UOnlineHotfixManager::~UOnlineHotfixManager()
 {
+	if (!UObject::IsPendingKillEnabled())
+	{
+		FCoreUObjectDelegates::GetPreGarbageCollectDelegate().RemoveAll(this);
+	}
 }
 
 UOnlineHotfixManager* UOnlineHotfixManager::Get(UWorld* World)
@@ -1891,6 +1900,11 @@ UWorld* UOnlineHotfixManager::GetWorld() const
 	return OwnerWorld.IsValid() ? OwnerWorld.Get() : nullptr;
 }
 
+void UOnlineHotfixManager::StopTrackingInvalidHotfixedAssets()
+{
+	AssetsHotfixedFromIniFiles.RemoveAllSwap([](const UObject* Obj) { return !IsValid(Obj); });
+}
+
 struct FHotfixManagerExec :
 	public FSelfRegisteringExec
 {
diff --git a/Engine/Plugins/Online/OnlineFramework/Source/Hotfix/Public/OnlineHotfixManager.h b/Engine/Plugins/Online/OnlineFramework/Source/Hotfix/Public/OnlineHotfixManager.h
--- a/Engine/Plugins/Online/OnlineFramework/Source/Hotfix/Public/OnlineHotfixManager.h
+++ b/Engine/Plugins/Online/OnlineFramework/Source/Hotfix/Public/OnlineHotfixManager.h
@@ -243,6 +243,9 @@ protected:
 	/** @return our current world */
 	UWorld* GetWorld() const override;
 
+	/** Stop tracking hotfixed assets marked as garbage */
+	void StopTrackingInvalidHotfixedAssets();
+	
 protected:
 
 	/**
