Extracted logic to reload objects affected by an ini file from UOnlineHotfixManager::HotfixIniFile into a seperate function.

#rnx

diff --git a/Engine/Plugins/Online/OnlineFramework/Source/Hotfix/Private/OnlineHotfixManager.cpp b/Engine/Plugins/Online/OnlineFramework/Source/Hotfix/Private/OnlineHotfixManager.cpp
--- a/Engine/Plugins/Online/OnlineFramework/Source/Hotfix/Private/OnlineHotfixManager.cpp
+++ b/Engine/Plugins/Online/OnlineFramework/Source/Hotfix/Private/OnlineHotfixManager.cpp
@@ -932,163 +932,9 @@
 	FConfigFileBackup& BackupFile = BackupIniFile(FileName, ConfigFile);
 	// Merge the string into the config file
 	ConfigFile->CombineFromBuffer(IniData, FileName);
-	TArray<UClass*> Classes;
-	TArray<UObject*> PerObjectConfigObjects;
-	int32 StartIndex = 0;
-	int32 EndIndex = 0;
-	TSet<FString> UpdatedSectionNames;
-	// Find the set of object classes that were affected
-	while (StartIndex >= 0 && StartIndex < IniData.Len() && EndIndex >= StartIndex)
-	{
-		// Find the next section header
-		StartIndex = IniData.Find(TEXT("["), ESearchCase::IgnoreCase, ESearchDir::FromStart, StartIndex);
-		if (StartIndex > -1)
-		{
-			// Find the ending section identifier
-			EndIndex = IniData.Find(TEXT("]"), ESearchCase::IgnoreCase, ESearchDir::FromStart, StartIndex);
-			if (EndIndex > StartIndex)
-			{
-				// Ignore square brackets in the middle of string
-				// - per object section starts with new line
-				// - there's no " character between opening bracket and line start
-				const bool bStartsWithNewLine = (StartIndex == 0) || (IniData[StartIndex - 1] == TEXT('\n'));
-				if (!bStartsWithNewLine)
-				{
-					bool bStartsInsideString = false;
-					for (int32 CharIdx = StartIndex - 1; CharIdx >= 0; CharIdx--)
-					{
-						const bool bHasStringMarker = (IniData[CharIdx] == TEXT('"'));
-						if (bHasStringMarker)
-						{
-							bStartsInsideString = true;
-							break;
-						}
 
-						const bool bHasNewLineMarker = (IniData[CharIdx] == TEXT('\n'));
-						if (bHasNewLineMarker)
-						{
-							break;
-						}
-					}
+	ReloadObjectsAffectedByConfigFile(FileName, IniData, ConfigFile->Name.ToString(), BackupFile.ClassesReloaded, false);
 
-					if (bStartsInsideString)
-					{
-						StartIndex = EndIndex;
-						continue;
-					}
-				}
-
-				UpdatedSectionNames.Emplace(IniData.Mid(StartIndex+1, EndIndex - StartIndex - 1));
-
-				int32 PerObjectNameIndex = IniData.Find(TEXT(" "), ESearchCase::IgnoreCase, ESearchDir::FromStart, StartIndex);
-
-				const TCHAR* AssetHotfixIniHACK = TEXT("[AssetHotfix]");
-				if (FCString::Strnicmp(*IniData + StartIndex, AssetHotfixIniHACK, FCString::Strlen(AssetHotfixIniHACK)) == 0)
-				{
-					// HACK - Make AssetHotfix the last element in the ini file so that this parsing isn't affected by it for now
-					break;
-				}
-
-				// Per object config entries will have a space in the name, but classes won't
-				if (PerObjectNameIndex == -1 || PerObjectNameIndex > EndIndex)
-				{
-					const TCHAR* ScriptHeader = TEXT("[/Script/");
-					const TCHAR* GameHeader = TEXT("[/Game/");
-					if (FCString::Strnicmp(*IniData + StartIndex, ScriptHeader, FCString::Strlen(ScriptHeader)) == 0)
-					{
-						const int32 ScriptSectionTag = 9;
-						// Snip the text out and try to find the class for that
-						const FString PackageClassName = IniData.Mid(StartIndex + ScriptSectionTag, EndIndex - StartIndex - ScriptSectionTag);
-						// Find the class for this so we know what to update
-						UClass* Class = FindObject<UClass>(nullptr, *PackageClassName, true);
-						if (Class)
-						{
-							// Add this to the list to check against
-							Classes.Add(Class);
-							BackupFile.ClassesReloaded.AddUnique(Class->GetPathName());
-						}
-					}
-					else if (FCString::Strnicmp(*IniData + StartIndex, GameHeader, FCString::Strlen(GameHeader)) == 0)
-					{
-						const int32 GameSectionTag = 1;
-						// Snip the text out and try to find the class for that
-						const FString PackageClassName = IniData.Mid(StartIndex + GameSectionTag, EndIndex - StartIndex - GameSectionTag);
-						UBlueprintGeneratedClass* BPGeneratedClass = LoadObject<UBlueprintGeneratedClass>(nullptr, *PackageClassName);
-						if (BPGeneratedClass)
-						{
-							// Add this to the list to check against
-							Classes.Add(BPGeneratedClass);
-							BackupFile.ClassesReloaded.AddUnique(BPGeneratedClass->GetPathName());
-						}
-					}
-				}
-				// Handle the per object config case by finding the object for reload
-				else
-				{
-					const int32 ClassNameStart = PerObjectNameIndex + 1;
-					const FString ClassName = IniData.Mid(ClassNameStart, EndIndex - ClassNameStart);
-
-					// Look up the class to search for
-					UClass* ObjectClass = UClass::TryFindTypeSlow<UClass>(ClassName);
-
-					if (ObjectClass)
-					{
-						const int32 Count = PerObjectNameIndex - StartIndex - 1;
-						const FString PerObjectName = IniData.Mid(StartIndex + 1, Count);
-
-						// Explicitly search the transient package (won't update non-transient objects)
-						UObject* PerObject = StaticFindFirstObject(ObjectClass, *PerObjectName, EFindFirstObjectOptions::NativeFirst);
-						if (PerObject != nullptr)
-						{
-							PerObjectConfigObjects.Add(PerObject);
-							BackupFile.ClassesReloaded.AddUnique(ObjectClass->GetPathName());
-						}
-					}
-					else
-					{
-						UE_LOG(LogHotfixManager, Warning, TEXT("Specified per-object class %s was not found"), *ClassName);
-					}
-				}
-				StartIndex = EndIndex;
-			}
-		}
-	}
-
-	int32 NumObjectsReloaded = 0;
-	const double StartTime = FPlatformTime::Seconds();
-	// Now that we have a list of classes to update, we can iterate objects and reload
-	for (UClass* Class : Classes)
-	{
-		if (Class->HasAnyClassFlags(CLASS_Config))
-		{
-			TArray<UObject*> Objects;
-			GetObjectsOfClass(Class, Objects, true, RF_NoFlags);
-			for (UObject* Object : Objects)
-			{
-				if (IsValid(Object))
-				{
-					// Force a reload of the config vars
-					UE_LOG(LogHotfixManager, Verbose, TEXT("Reloading %s"), *Object->GetPathName());
-					Object->ReloadConfig();
-					NumObjectsReloaded++;
-				}
-			}
-		}
-	}
-
-	// Reload any PerObjectConfig objects that were affected
-	for (auto ReloadObject : PerObjectConfigObjects)
-	{
-		UE_LOG(LogHotfixManager, Verbose, TEXT("Reloading %s"), *ReloadObject->GetPathName());
-		ReloadObject->ReloadConfig();
-		NumObjectsReloaded++;
-	}
-
-	const FString ConfigFileName = ConfigFile->Name.ToString();
-	FCoreDelegates::TSOnConfigSectionsChanged().Broadcast(ConfigFileName, UpdatedSectionNames);
-
-	UE_LOG(LogHotfixManager, Log, TEXT("Updating config from %s took %f seconds and reloaded %d objects"),
-		*FileName, FPlatformTime::Seconds() - StartTime, NumObjectsReloaded);
 	return true;
 }
 
@@ -1905,6 +1751,166 @@
 	AssetsHotfixedFromIniFiles.RemoveAllSwap([](const UObject* Obj) { return !IsValid(Obj); });
 }
 
+void UOnlineHotfixManager::ReloadObjectsAffectedByConfigFile(const FString& IniDataFileName, const FString& IniData, const FString& ConfigFilename, TArray<FString>& ReloadedClassesPathNames, bool bUseLoadConfig)
+{
+	TArray<UClass*> Classes;
+	TArray<UObject*> PerObjectConfigObjects;
+	int32 StartIndex = 0;
+	int32 EndIndex = 0;
+	TSet<FString> UpdatedSectionNames;
+	// Find the set of object classes that were affected
+	while (StartIndex >= 0 && StartIndex < IniData.Len() && EndIndex >= StartIndex)
+	{
+		// Find the next section header
+		StartIndex = IniData.Find(TEXT("["), ESearchCase::IgnoreCase, ESearchDir::FromStart, StartIndex);
+		if (StartIndex > -1)
+		{
+			// Find the ending section identifier
+			EndIndex = IniData.Find(TEXT("]"), ESearchCase::IgnoreCase, ESearchDir::FromStart, StartIndex);
+			if (EndIndex > StartIndex)
+			{
+				// Ignore square brackets in the middle of string
+				// - per object section starts with new line
+				// - there's no " character between opening bracket and line start
+				const bool bStartsWithNewLine = (StartIndex == 0) || (IniData[StartIndex - 1] == TEXT('\n'));
+				if (!bStartsWithNewLine)
+				{
+					bool bStartsInsideString = false;
+					for (int32 CharIdx = StartIndex - 1; CharIdx >= 0; CharIdx--)
+					{
+						const bool bHasStringMarker = (IniData[CharIdx] == TEXT('"'));
+						if (bHasStringMarker)
+						{
+							bStartsInsideString = true;
+							break;
+						}
+
+						const bool bHasNewLineMarker = (IniData[CharIdx] == TEXT('\n'));
+						if (bHasNewLineMarker)
+						{
+							break;
+						}
+					}
+
+					if (bStartsInsideString)
+					{
+						StartIndex = EndIndex;
+						continue;
+					}
+				}
+
+				UpdatedSectionNames.Emplace(IniData.Mid(StartIndex + 1, EndIndex - StartIndex - 1));
+
+				int32 PerObjectNameIndex = IniData.Find(TEXT(" "), ESearchCase::IgnoreCase, ESearchDir::FromStart, StartIndex);
+
+				const TCHAR* AssetHotfixIniHACK = TEXT("[AssetHotfix]");
+				if (FCString::Strnicmp(*IniData + StartIndex, AssetHotfixIniHACK, FCString::Strlen(AssetHotfixIniHACK)) == 0)
+				{
+					// HACK - Make AssetHotfix the last element in the ini file so that this parsing isn't affected by it for now
+					break;
+				}
+
+				// Per object config entries will have a space in the name, but classes won't
+				if (PerObjectNameIndex == -1 || PerObjectNameIndex > EndIndex)
+				{
+					const TCHAR* ScriptHeader = TEXT("[/Script/");
+					const TCHAR* GameHeader = TEXT("[/Game/");
+					if (FCString::Strnicmp(*IniData + StartIndex, ScriptHeader, FCString::Strlen(ScriptHeader)) == 0)
+					{
+						const int32 ScriptSectionTag = 9;
+						// Snip the text out and try to find the class for that
+						const FString PackageClassName = IniData.Mid(StartIndex + ScriptSectionTag, EndIndex - StartIndex - ScriptSectionTag);
+						// Find the class for this so we know what to update
+						UClass* Class = FindObject<UClass>(nullptr, *PackageClassName, true);
+						if (Class)
+						{
+							// Add this to the list to check against
+							Classes.Add(Class);
+							ReloadedClassesPathNames.AddUnique(Class->GetPathName());
+						}
+					}
+					else if (FCString::Strnicmp(*IniData + StartIndex, GameHeader, FCString::Strlen(GameHeader)) == 0)
+					{
+						const int32 GameSectionTag = 1;
+						// Snip the text out and try to find the class for that
+						const FString PackageClassName = IniData.Mid(StartIndex + GameSectionTag, EndIndex - StartIndex - GameSectionTag);
+						UBlueprintGeneratedClass* BPGeneratedClass = LoadObject<UBlueprintGeneratedClass>(nullptr, *PackageClassName);
+						if (BPGeneratedClass)
+						{
+							// Add this to the list to check against
+							Classes.Add(BPGeneratedClass);
+							ReloadedClassesPathNames.AddUnique(BPGeneratedClass->GetPathName());
+						}
+					}
+				}
+				// Handle the per object config case by finding the object for reload
+				else
+				{
+					const int32 ClassNameStart = PerObjectNameIndex + 1;
+					const FString ClassName = IniData.Mid(ClassNameStart, EndIndex - ClassNameStart);
+
+					// Look up the class to search for
+					UClass* ObjectClass = UClass::TryFindTypeSlow<UClass>(ClassName);
+
+					if (ObjectClass)
+					{
+						const int32 Count = PerObjectNameIndex - StartIndex - 1;
+						const FString PerObjectName = IniData.Mid(StartIndex + 1, Count);
+
+						// Explicitly search the transient package (won't update non-transient objects)
+						UObject* PerObject = StaticFindFirstObject(ObjectClass, *PerObjectName, EFindFirstObjectOptions::NativeFirst);
+						if (PerObject != nullptr)
+						{
+							PerObjectConfigObjects.Add(PerObject);
+							ReloadedClassesPathNames.AddUnique(ObjectClass->GetPathName());
+						}
+					}
+					else
+					{
+						UE_LOG(LogHotfixManager, Warning, TEXT("Specified per-object class %s was not found"), *ClassName);
+					}
+				}
+				StartIndex = EndIndex;
+			}
+		}
+	}
+
+	int32 NumObjectsReloaded = 0;
+	const double StartTime = FPlatformTime::Seconds();
+	// Now that we have a list of classes to update, we can iterate objects and reload
+	for (UClass* Class : Classes)
+	{
+		if (Class->HasAnyClassFlags(CLASS_Config))
+		{
+			TArray<UObject*> Objects;
+			GetObjectsOfClass(Class, Objects, true, RF_NoFlags);
+			for (UObject* Object : Objects)
+			{
+				if (IsValid(Object))
+				{
+					// Force a reload of the config vars
+					UE_LOG(LogHotfixManager, Verbose, TEXT("Reloading %s"), *Object->GetPathName());
+					bUseLoadConfig ? Object->LoadConfig() : Object->ReloadConfig();
+					NumObjectsReloaded++;
+				}
+			}
+		}
+	}
+
+	// Reload any PerObjectConfig objects that were affected
+	for (UObject* ReloadObject : PerObjectConfigObjects)
+	{
+		UE_LOG(LogHotfixManager, Verbose, TEXT("Reloading %s"), *ReloadObject->GetPathName());
+		bUseLoadConfig ? ReloadObject->LoadConfig() : ReloadObject->ReloadConfig();
+		NumObjectsReloaded++;
+	}
+
+	FCoreDelegates::TSOnConfigSectionsChanged().Broadcast(ConfigFilename, UpdatedSectionNames);
+
+	UE_LOG(LogHotfixManager, Log, TEXT("Updating config from %s took %f seconds and reloaded %d objects"),
+		*IniDataFileName, FPlatformTime::Seconds() - StartTime, NumObjectsReloaded);
+}
+
 struct FHotfixManagerExec :
 	public FSelfRegisteringExec
 {
diff --git a/Engine/Plugins/Online/OnlineFramework/Source/Hotfix/Public/OnlineHotfixManager.h b/Engine/Plugins/Online/OnlineFramework/Source/Hotfix/Public/OnlineHotfixManager.h
--- a/Engine/Plugins/Online/OnlineFramework/Source/Hotfix/Public/OnlineHotfixManager.h
+++ b/Engine/Plugins/Online/OnlineFramework/Source/Hotfix/Public/OnlineHotfixManager.h
@@ -405,6 +405,8 @@ public:
 
 	/** Factory method that returns the configured hotfix manager */
 	static UOnlineHotfixManager* Get(UWorld* World);
+
+	static void ReloadObjectsAffectedByConfigFile(const FString& IniDataFileName, const FString& IniData, const FString& ConfigFilename, TArray<FString>& ReloadedClassesPathNames, bool bUseLoadConfig);
 };
 
 #if UE_ENABLE_INCLUDE_ORDER_DEPRECATED_IN_5_2
