Add command line argument to test hotfixing

Add a command line argument "-TestHotfixIniFile" to test hotfixing a configuration file. This allows the timing to be identical to regular hotfixing.

diff --git a/Engine/Plugins/Online/OnlineFramework/Source/Hotfix/Private/OnlineHotfixManager.cpp b/Engine/Plugins/Online/OnlineFramework/Source/Hotfix/Private/OnlineHotfixManager.cpp
--- a/Engine/Plugins/Online/OnlineFramework/Source/Hotfix/Private/OnlineHotfixManager.cpp
+++ b/Engine/Plugins/Online/OnlineFramework/Source/Hotfix/Private/OnlineHotfixManager.cpp
@@ -747,8 +747,40 @@ EHotfixResult UOnlineHotfixManager::ApplyHotfix()
 	return Result;
 }
 
+#if !UE_BUILD_SHIPPING
+void UOnlineHotfixManager::ApplyLocalTestHotfix(FString Filename)
+{
+	const FString CleanName = FPaths::GetCleanFilename(Filename);
+	FString IniData;
+	if (FFileHelper::LoadFileToString(IniData, *Filename))
+	{
+		if (HotfixIniFile(CleanName, IniData))
+		{
+			UE_LOG(LogHotfixManager, Display, TEXT("Successfully applied test hotfix file %s"), *Filename);
+		}
+		else
+		{
+			UE_LOG(LogHotfixManager, Display, TEXT("Failed to apply test hotfix file %s"), *Filename);
+		}
+	}
+	else
+	{
+		UE_LOG(LogHotfixManager, Warning, TEXT("Unable to read test hotfix file '%s'"), *Filename);
+	}
+}
+#endif
+
 void UOnlineHotfixManager::TriggerHotfixComplete(EHotfixResult HotfixResult)
 {
+#if !UE_BUILD_SHIPPING
+	// Apply this here so it overwrites any downloaded hotfix changes
+	FString IniFilename;
+	if (FParse::Value(FCommandLine::Get(), TEXT("-TestHotfixIniFile="), IniFilename))
+	{
+		ApplyLocalTestHotfix(IniFilename);
+	}
+#endif
+	
 	if (HotfixResult != EHotfixResult::Failed && HotfixResult != EHotfixResult::SuccessNoChange)
 	{
 		PatchAssetsFromIniFiles();
diff --git a/Engine/Plugins/Online/OnlineFramework/Source/Hotfix/Public/OnlineHotfixManager.h b/Engine/Plugins/Online/OnlineFramework/Source/Hotfix/Public/OnlineHotfixManager.h
--- a/Engine/Plugins/Online/OnlineFramework/Source/Hotfix/Public/OnlineHotfixManager.h
+++ b/Engine/Plugins/Online/OnlineFramework/Source/Hotfix/Public/OnlineHotfixManager.h
@@ -372,6 +372,11 @@ protected:
 	/** Allow child classes to determine if specific assets should be hotfixed or not */
 	virtual bool ShouldHotfixAsset(const FString& AssetPath) const;
 
+#if !UE_BUILD_SHIPPING
+	/** Test function that applies a local file as if it were a hotfix. */
+	void ApplyLocalTestHotfix(FString Filename);
+#endif
+
 public:
 	UOnlineHotfixManager();
 	UOnlineHotfixManager(FVTableHelper& Helper);
