Update patch check flow to detect the platform environment earlier in the patching process for supported platforms

diff --git a/Engine/Plugins/Online/OnlineFramework/Source/Hotfix/Hotfix.Build.cs b/Engine/Plugins/Online/OnlineFramework/Source/Hotfix/Hotfix.Build.cs
--- a/Engine/Plugins/Online/OnlineFramework/Source/Hotfix/Hotfix.Build.cs
+++ b/Engine/Plugins/Online/OnlineFramework/Source/Hotfix/Hotfix.Build.cs
@@ -26,9 +26,5 @@ public class Hotfix : ModuleRules
 				"OnlineSubsystem",
 			}
 			);
-
-		PublicDefinitions.Add("UPDATEMANAGER_PLATFORM_ENVIRONMENT_DETECTION=" + (bPlatformEnvironmentDetection ? "1" : "0"));
 	}
-
-	protected virtual bool bPlatformEnvironmentDetection { get { return false; } }
 }
diff --git a/Engine/Plugins/Online/OnlineFramework/Source/Hotfix/Private/UpdateManager.cpp b/Engine/Plugins/Online/OnlineFramework/Source/Hotfix/Private/UpdateManager.cpp
--- a/Engine/Plugins/Online/OnlineFramework/Source/Hotfix/Private/UpdateManager.cpp
+++ b/Engine/Plugins/Online/OnlineFramework/Source/Hotfix/Private/UpdateManager.cpp
@@ -144,7 +144,6 @@ UUpdateManager::EUpdateStartResult UUpdateManager::StartUpdateCheckInternal(cons
 	{
 		bCheckHotfixAvailabilityOnly = ContextDefinition.bCheckAvailabilityOnly;
 		bCurrentUpdatePatchCheckEnabled = ContextDefinition.bPatchCheckEnabled;
-		bCurrentUpdatePlatformEnvironmentDetectionEnabled = ContextDefinition.bPlatformEnvironmentDetectionEnabled;
 
 		// Immediately move into a pending state so the UI state trigger fires
 		SetUpdateState(EUpdateState::UpdatePending);
@@ -245,7 +244,6 @@ void UUpdateManager::CheckComplete(EUpdateCompletionStatus Result, bool bUpdateT
 
 		bCheckHotfixAvailabilityOnly = false;
 		bCurrentUpdatePatchCheckEnabled = true;
-		bCurrentUpdatePlatformEnvironmentDetectionEnabled = true;
 		
 		OnUpdateCheckComplete().Broadcast(FinalResult);
 	};
@@ -352,7 +350,7 @@ void UUpdateManager::InstallBundlePatchCheckComplete(EInstallBundleManagerPatchC
 
 	if (PatchResult == EInstallBundleManagerPatchCheckResult::NoPatchRequired)
 	{
-		StartPlatformEnvironmentCheck();
+		StartHotfixCheck();
 	}
 	else if (PatchResult == EInstallBundleManagerPatchCheckResult::NoLoggedInUser)
 	{
@@ -369,53 +367,6 @@ void UUpdateManager::InstallBundlePatchCheckComplete(EInstallBundleManagerPatchC
 	}
 }
 
-void UUpdateManager::StartPlatformEnvironmentCheck()
-{
-	if (bCurrentUpdatePlatformEnvironmentDetectionEnabled && !bPlatformEnvironmentDetected)
-	{
-#if UPDATEMANAGER_PLATFORM_ENVIRONMENT_DETECTION
-		if (DetectPlatformEnvironment())
-		{
-			return;
-		}
-#endif
-	}
-	
-	StartHotfixCheck();
-}
-
-void UUpdateManager::OnDetectPlatformEnvironmentComplete(const FOnlineError& Result)
-{
-	if (Result.WasSuccessful())
-	{
-		bPlatformEnvironmentDetected = true;
-		StartHotfixCheck();
-	}
-	else
-	{
-		if (Result.GetErrorCode().Contains(TEXT("getUserAccessCode failed : 0x8055000f"), ESearchCase::IgnoreCase))
-		{
-			UE_LOG(LogHotfixManager, Warning, TEXT("Failed to complete login because patch is required"));
-			CheckComplete(EUpdateCompletionStatus::UpdateSuccess_NeedsPatch);
-		}
-		else
-		{
-			if (Result.GetErrorCode().Contains(TEXT("com.epicgames.identity.notloggedin"), ESearchCase::IgnoreCase))
-			{
-				UE_LOG(LogHotfixManager, Warning, TEXT("Failed to detect online environment for the platform, no user signed in"));
-				CheckComplete(EUpdateCompletionStatus::UpdateFailure_NotLoggedIn);
-			}
-			else
-			{
-				// just a platform env error, assume production and keep going
-				UE_LOG(LogHotfixManager, Warning, TEXT("Failed to detect online environment for the platform"));
-				bPlatformEnvironmentDetected = true;
-				StartHotfixCheck();
-			}
-		}
-	}
-}
-
 void UUpdateManager::StartHotfixCheck()
 {
 	if (bCheckHotfixAvailabilityOnly)
diff --git a/Engine/Plugins/Online/OnlineFramework/Source/Hotfix/Public/UpdateManager.h b/Engine/Plugins/Online/OnlineFramework/Source/Hotfix/Public/UpdateManager.h
--- a/Engine/Plugins/Online/OnlineFramework/Source/Hotfix/Public/UpdateManager.h
+++ b/Engine/Plugins/Online/OnlineFramework/Source/Hotfix/Public/UpdateManager.h
@@ -27,8 +27,6 @@ enum class EUpdateState : uint8
 	UpdatePending,
 	/** Checking for an available patch */
 	CheckingForPatch,
-	/** Detect the console environment via an auth request */
-	DetectingPlatformEnvironment,
 	/** Checking with hotfix for available updates */
 	CheckingForHotfix,
 	/** Waiting for the async loading / preloading to complete */
@@ -125,9 +123,6 @@ public:
 	/** Whether to check for patches during the update process. */
 	UPROPERTY(Config)
 	bool bPatchCheckEnabled = true;
-	/** Whether to attempt detection of the platforms environment. */
-	UPROPERTY(Config)
-	bool bPlatformEnvironmentDetectionEnabled = true;
 	/** Additional tags for extensibility. */
 	UPROPERTY(Config)
 	TSet<FString> AdditionalTags;
@@ -333,31 +328,6 @@ protected:
 	/** Availability check complete */
 	virtual void HotfixAvailabilityCheckComplete(EHotfixResult Result);
 
-	/**
-	 * Query for platform environment before continuing with hotfixing
-	 * This configures the Mcp backend correctly
-	 */
-	virtual void StartPlatformEnvironmentCheck();
-
-#if UPDATEMANAGER_PLATFORM_ENVIRONMENT_DETECTION
-	/**
-	 * Platform specific implementation of platform environment detection
-	 * @return true if the detection began.  false if the detection did not begin and we should continue the checks.
-	 */
-	bool DetectPlatformEnvironment();
-
-	/**
-	 * Platform specific callback for logging in on console which is needed for the platform environment detection
-	 */
-	virtual void DetectPlatformEnvironment_OnLoginConsoleComplete(int32 LocalUserNum, bool bWasSuccessful, const FUniqueNetId& UserId, const FString& Error);
-#endif
-
-	/**
-	 * Callback when detecting platform environment completes
-	 */
-	void OnDetectPlatformEnvironmentComplete(const FOnlineError& Result);
-
-
 	/**
 	 * Preload game assets after patch/hotfix check is complete but before game is alerted
 	 */
@@ -397,10 +367,6 @@ protected:
 	virtual void OnApplicationHasReactivated();
 
 protected:
-	/** true if we've already detected the backend environment */
-	UPROPERTY()
-	bool bPlatformEnvironmentDetected = !UPDATEMANAGER_PLATFORM_ENVIRONMENT_DETECTION; // Default to true if we do not need to detect
-
 	/** Has the first update completed */
 	UPROPERTY()
 	bool bInitialUpdateFinished;
@@ -427,9 +393,6 @@ protected:
 	FDelegateHandle HotfixCompleteDelegateHandle;
 	FDelegateHandle HotfixProgressDelegateHandle;
 	FDelegateHandle HotfixProcessedFileDelegateHandle;
-#if UPDATEMANAGER_PLATFORM_ENVIRONMENT_DETECTION
-	FDelegateHandle OnLoginConsoleCompleteHandle;
-#endif
 
 	/** The time at which we started the initial load after updates completed */
 	double LoadStartTime;
@@ -478,7 +441,6 @@ private:
 	TMap<FString, FUpdateContextDefinition> ProcessedUpdateContextDefinitions;
 
 	bool bCurrentUpdatePatchCheckEnabled = true;
-	bool bCurrentUpdatePlatformEnvironmentDetectionEnabled = true;
 protected:
 
 	/** @return a pointer to the hotfix manager */
diff --git a/Engine/Plugins/Online/OnlineFramework/Source/PatchCheck/PatchCheck.build.cs b/Engine/Plugins/Online/OnlineFramework/Source/PatchCheck/PatchCheck.build.cs
--- a/Engine/Plugins/Online/OnlineFramework/Source/PatchCheck/PatchCheck.build.cs
+++ b/Engine/Plugins/Online/OnlineFramework/Source/PatchCheck/PatchCheck.build.cs
@@ -23,7 +23,9 @@ public class PatchCheck : ModuleRules
 		);
 
 		PrivateDefinitions.Add("PATCH_CHECK_FAIL_ON_GENERIC_FAILURE=" + (bFailOnGenericFailure ? "1" : "0"));
+		PublicDefinitions.Add("PATCH_CHECK_PLATFORM_ENVIRONMENT_DETECTION=" + (bPlatformEnvironmentDetection ? "1" : "0"));
 	}
 
 	protected virtual bool bFailOnGenericFailure { get { return true; } }
+	protected virtual bool bPlatformEnvironmentDetection { get { return false; } }
 }
diff --git a/Engine/Plugins/Online/OnlineFramework/Source/PatchCheck/Private/PatchCheck.cpp b/Engine/Plugins/Online/OnlineFramework/Source/PatchCheck/Private/PatchCheck.cpp
--- a/Engine/Plugins/Online/OnlineFramework/Source/PatchCheck/Private/PatchCheck.cpp
+++ b/Engine/Plugins/Online/OnlineFramework/Source/PatchCheck/Private/PatchCheck.cpp
@@ -70,6 +70,40 @@ void FPatchCheck::RefreshConfig()
 			ensureMsgf(false, TEXT("UpdateManager::bCheckOSSForUpdate is deprecated, Set FPatchCheck::bCheckOSSForUpdate using section [PatchCheck] instead."));
 		}
 	}
+
+	GConfig->GetBool(TEXT("PatchCheck"), TEXT("bPlatformEnvironmentDetectionEnabled"), bPlatformEnvironmentDetectionEnabled, GEngineIni);
+}
+
+void FPatchCheck::OnDetectPlatformEnvironmentComplete(const FOnlineError& Result)
+{
+	if (Result.WasSuccessful())
+	{
+		bPlatformEnvironmentDetected = true;
+		HandleOSSPatchCheck();
+	}
+	else
+	{
+		if (Result.GetErrorCode().Contains(TEXT("getUserAccessCode failed : 0x8055000f"), ESearchCase::IgnoreCase))
+		{
+			UE_LOG(LogPatchCheck, Warning, TEXT("Failed to complete login because patch is required"));
+			PatchCheckComplete(EPatchCheckResult::PatchRequired);
+		}
+		else
+		{
+			if (Result.GetErrorCode().Contains(TEXT("com.epicgames.identity.notloggedin"), ESearchCase::IgnoreCase))
+			{
+				UE_LOG(LogPatchCheck, Warning, TEXT("Failed to detect online environment for the platform, no user signed in"));
+				PatchCheckComplete(EPatchCheckResult::NoLoggedInUser);
+			}
+			else
+			{
+				// just a platform env error, assume production and keep going
+				UE_LOG(LogPatchCheck, Warning, TEXT("Failed to detect online environment for the platform"));
+				bPlatformEnvironmentDetected = true;
+				HandleOSSPatchCheck();
+			}
+		}
+	}
 }
 
 void FPatchCheck::StartPatchCheck()
@@ -79,6 +113,21 @@ void FPatchCheck::StartPatchCheck()
 
 	RefreshConfig();
 
+	if (bPlatformEnvironmentDetectionEnabled && !bPlatformEnvironmentDetected)
+	{
+#if PATCH_CHECK_PLATFORM_ENVIRONMENT_DETECTION
+		if (DetectPlatformEnvironment())
+		{
+			return;
+		}
+#endif
+	}
+
+	HandleOSSPatchCheck();
+}
+
+void FPatchCheck::HandleOSSPatchCheck()
+{
 	if (bCheckPlatformOSSForUpdate && IOnlineSubsystem::GetByPlatform() != nullptr)
 	{
 		bIsCheckInProgress = true;
diff --git a/Engine/Plugins/Online/OnlineFramework/Source/PatchCheck/Public/PatchCheck.h b/Engine/Plugins/Online/OnlineFramework/Source/PatchCheck/Public/PatchCheck.h
--- a/Engine/Plugins/Online/OnlineFramework/Source/PatchCheck/Public/PatchCheck.h
+++ b/Engine/Plugins/Online/OnlineFramework/Source/PatchCheck/Public/PatchCheck.h
@@ -56,11 +56,29 @@ protected:
 	void RefreshConfig();
 	virtual void StartPlatformOSSPatchCheck();
 	virtual void StartOSSPatchCheck();
+	void HandleOSSPatchCheck();
 	virtual bool EnvironmentWantsPatchCheck() const;
 	bool SkipPatchCheck() const;
 	void OnCheckForPatchComplete(const FUniqueNetId& UniqueId, EUserPrivileges::Type Privilege, uint32 PrivilegeResult, bool bConsoleCheck);
 	virtual void PatchCheckComplete(EPatchCheckResult PatchResult);
 
+#if PATCH_CHECK_PLATFORM_ENVIRONMENT_DETECTION
+	/**
+	 * Platform specific implementation of platform environment detection
+	 * @return true if the detection began.  false if the detection did not begin and we should continue the checks.
+	 */
+	bool DetectPlatformEnvironment();
+
+	/**
+	 * Platform specific callback for logging in on console which is needed for the platform environment detection
+	 */
+	virtual void DetectPlatformEnvironment_OnLoginConsoleComplete(int32 LocalUserNum, bool bWasSuccessful, const FUniqueNetId& UserId, const FString& Error);
+#endif
+	/**
+	 * Callback when detecting platform environment completes
+	 */
+	void OnDetectPlatformEnvironmentComplete(const FOnlineError& Result);
+
 protected:
 
 	FOnPatchCheckComplete OnComplete;
@@ -76,5 +94,13 @@ protected:
 	/** Check the default OSS for an update */
 	bool bCheckOSSForUpdate = true;
 
+#if PATCH_CHECK_PLATFORM_ENVIRONMENT_DETECTION
+	FDelegateHandle OnLoginConsoleCompleteHandle;
+#endif
+
+	/** true if we've already detected the backend environment */
+	bool bPlatformEnvironmentDetected = !PATCH_CHECK_PLATFORM_ENVIRONMENT_DETECTION; // Default to true if we do not need to detect
+	bool bPlatformEnvironmentDetectionEnabled = true;
+
 	friend class TPatchCheckModule<FPatchCheck>;
 };
