Detect failure to parse corrupt hotfix files with early terminators.

diff --git a/Engine/Plugins/Online/OnlineFramework/Source/Hotfix/Private/OnlineHotfixManager.cpp b/Engine/Plugins/Online/OnlineFramework/Source/Hotfix/Private/OnlineHotfixManager.cpp
--- a/Engine/Plugins/Online/OnlineFramework/Source/Hotfix/Private/OnlineHotfixManager.cpp
+++ b/Engine/Plugins/Online/OnlineFramework/Source/Hotfix/Private/OnlineHotfixManager.cpp
@@ -1376,7 +1376,11 @@ void UOnlineHotfixManager::PatchAssetsFromIniFiles()
 
 								// We have to read json data as quoted string because tokenizing it creates extra unwanted characters.
 								FString JsonData;
-								if (FParse::QuotedString(*Tokens[2], JsonData))
+
+								// Json should be read in its entirety, if the whole buffer wasn't read the string is malformed. 
+								int32 ReadLen = 0;
+								int32 InputLen = Tokens[2].Len();
+								if (FParse::QuotedString(*Tokens[2], JsonData, &ReadLen) && ReadLen == InputLen)
 								{
 									HotfixTableUpdate(Asset, AssetPath, JsonData, ProblemStrings);
 									bAddAssetToHotfixedList = ProblemStrings.Num() == 0;
