Ensure data table row update hotfixes propagate data table changes for all supported property types, i.e. bool

diff --git a/Engine/Plugins/Online/OnlineFramework/Source/Hotfix/Private/OnlineHotfixManager.cpp b/Engine/Plugins/Online/OnlineFramework/Source/Hotfix/Private/OnlineHotfixManager.cpp
--- a/Engine/Plugins/Online/OnlineFramework/Source/Hotfix/Private/OnlineHotfixManager.cpp
+++ b/Engine/Plugins/Online/OnlineFramework/Source/Hotfix/Private/OnlineHotfixManager.cpp
@@ -1600,7 +1600,7 @@ void UOnlineHotfixManager::HotfixRowUpdate(UObject* Asset, const FString& AssetP
 							ProblemStrings.Add(Problem);
 						}
 					}
-										// String property
+					// String property
 					else if (StrProp)
 					{
 						const FString OldPropertyValue = StrProp->GetPropertyValue(RowData);
@@ -1630,17 +1630,22 @@ void UOnlineHotfixManager::HotfixRowUpdate(UObject* Asset, const FString& AssetP
 						bWasDataTableChanged = true;
 						UE_LOG(LogHotfixManager, Log, TEXT("Data table %s row %s updated column %s from %s to %s."), *AssetPath, *RowName, *ColumnName, *OldPropertyValue.ToString(), *NewPropertyValue.ToString());
 					}
-					// Not an expected property.
 					else
 					{
-						// we'll make one last attempt here
+						// Evaluate supported property types, i.e. FBoolProperty, and attempt to assign the value
+						const FString OldPropertyValue = DataTableUtils::GetPropertyValueAsString(DataTableRowProperty, (uint8*)DataTableRow, EDataTableExportFlags::UseSimpleText);
 						FString Error = DataTableUtils::AssignStringToProperty(NewValue, DataTableRowProperty, (uint8*)DataTableRow);
 
-						if (Error.Len() > 0)
+						if (Error.IsEmpty())
 						{
-							const FString Problem(FString::Printf(TEXT("The data table row property named %s is not a FNumericProperty, FStrProperty, FNameProperty, or FSoftObjectProperty and it should be."), *ColumnName));
-							ProblemStrings.Add(Problem);
-							ProblemStrings.Add(FString::Printf(TEXT("%s"), *Error));
+							bWasDataTableChanged = true;
+							UE_LOG(LogHotfixManager, Log, TEXT("Data table %s row %s updated column %s from %s to %s."), *AssetPath, *RowName, *ColumnName, *OldPropertyValue, *NewValue);
+						}
+						else
+						{
+							FString Problem(FString::Printf(TEXT("Failed to update data table %s row %s column %s from %s to %s."), *AssetPath, *RowName, *ColumnName, *OldPropertyValue, *NewValue));
+							ProblemStrings.Add(MoveTemp(Problem));
+							ProblemStrings.Add(MoveTemp(Error));
 						}
 					}
 				}
