Comment out ensure when LocalUser does not have UserInfo while we look into how to suppress it based on whether it's expected or not

diff --git a/Engine/Plugins/Online/OnlineFramework/Source/Party/Private/User/SocialUser.cpp b/Engine/Plugins/Online/OnlineFramework/Source/Party/Private/User/SocialUser.cpp
--- a/Engine/Plugins/Online/OnlineFramework/Source/Party/Private/User/SocialUser.cpp
+++ b/Engine/Plugins/Online/OnlineFramework/Source/Party/Private/User/SocialUser.cpp
@@ -353,7 +353,7 @@ void USocialUser::TryBroadcastInitializationComplete()
 			}
 			else
 			{
-				ensureAlwaysMsgf(!IsLocalUser(), TEXT("%s LocalUser %s does not have UserInfo!"), ANSI_TO_TCHAR(__FUNCTION__), *ToDebugString());
+				//ensureAlwaysMsgf(!IsLocalUser(), TEXT("%s LocalUser %s does not have UserInfo!"), ANSI_TO_TCHAR(__FUNCTION__), *ToDebugString());
 				UE_LOG(LogParty, Warning, TEXT("%s User %s does not have UserInfo!"), ANSI_TO_TCHAR(__FUNCTION__), *ToDebugString());
 
 				// User is invalid with no open queries
