Refactoring HotfixRowUpdate to be support hotfixing Blueprint Structs inside Data Tables.

diff --git a/Engine/Plugins/Online/OnlineFramework/Source/Hotfix/Private/OnlineHotfixManager.cpp b/Engine/Plugins/Online/OnlineFramework/Source/Hotfix/Private/OnlineHotfixManager.cpp
--- a/Engine/Plugins/Online/OnlineFramework/Source/Hotfix/Private/OnlineHotfixManager.cpp
+++ b/Engine/Plugins/Online/OnlineFramework/Source/Hotfix/Private/OnlineHotfixManager.cpp
@@ -22,6 +22,7 @@
 #include "Engine/BlueprintGeneratedClass.h"
 
 #include "Serialization/AsyncLoadingFlushContext.h"
+#include "UObject/PropertyAccessUtil.h"
 
 #include UE_INLINE_GENERATED_CPP_BY_NAME(OnlineHotfixManager)
 
@@ -1550,7 +1551,8 @@ void UOnlineHotfixManager::HotfixRowUpdate(UObject* Asset, const FString& AssetP
 	{
 		// Edit the row with the new value.
 		bool bWasDataTableChanged = false;
-		FProperty* DataTableRowProperty = DataTable->GetRowStruct()->FindPropertyByName(FName(*ColumnName));
+
+		FProperty* DataTableRowProperty = PropertyAccessUtil::FindPropertyByName(FName(*ColumnName), DataTable->GetRowStruct());
 		if (DataTableRowProperty)
 		{
 			// See what type of property this is.
@@ -1561,7 +1563,7 @@ void UOnlineHotfixManager::HotfixRowUpdate(UObject* Asset, const FString& AssetP
 
 			// Get the row data by name.
 			static const FString Context = FString(TEXT("UOnlineHotfixManager::PatchAssetsFromIniFiles"));
-			FTableRowBase* DataTableRow = DataTable->FindRow<FTableRowBase>(FName(*RowName), Context);
+			uint8* DataTableRow = DataTable->FindRowUnchecked(FName(*RowName));
 			if (DataTableRow)
 			{
 				uint8* RowData = DataTableRowProperty->ContainerPtrToValuePtr<uint8>(DataTableRow, 0);
