Clean up patch check logic for platform extensions.

#rnx

diff --git a/Engine/Plugins/Online/OnlineFramework/Source/PatchCheck/PatchCheck.build.cs b/Engine/Plugins/Online/OnlineFramework/Source/PatchCheck/PatchCheck.build.cs
--- a/Engine/Plugins/Online/OnlineFramework/Source/PatchCheck/PatchCheck.build.cs
+++ b/Engine/Plugins/Online/OnlineFramework/Source/PatchCheck/PatchCheck.build.cs
@@ -24,8 +24,10 @@ public class PatchCheck : ModuleRules
 
 		PrivateDefinitions.Add("PATCH_CHECK_FAIL_ON_GENERIC_FAILURE=" + (bFailOnGenericFailure ? "1" : "0"));
 		PublicDefinitions.Add("PATCH_CHECK_PLATFORM_ENVIRONMENT_DETECTION=" + (bPlatformEnvironmentDetection ? "1" : "0"));
+		PrivateDefinitions.Add("PATCH_CHECK_PRIVILEGE_MUST_BE_LOGGED_IN=" + (bCanCheckPlayOnlinePrivilegeMustBeLoggedIn ? "1" : "0"));
 	}
 
 	protected virtual bool bFailOnGenericFailure { get { return true; } }
 	protected virtual bool bPlatformEnvironmentDetection { get { return false; } }
+	protected virtual bool bCanCheckPlayOnlinePrivilegeMustBeLoggedIn { get { return true; } }
 }
diff --git a/Engine/Plugins/Online/OnlineFramework/Source/PatchCheck/Private/PatchCheck.cpp b/Engine/Plugins/Online/OnlineFramework/Source/PatchCheck/Private/PatchCheck.cpp
--- a/Engine/Plugins/Online/OnlineFramework/Source/PatchCheck/Private/PatchCheck.cpp
+++ b/Engine/Plugins/Online/OnlineFramework/Source/PatchCheck/Private/PatchCheck.cpp
@@ -165,8 +165,8 @@ void FPatchCheck::StartPlatformOSSPatchCheck()
 	if (PlatformOnlineIdentity.IsValid())
 	{
 		FUniqueNetIdPtr UserId = GetFirstSignedInUser(PlatformOnlineIdentity);
-#if PLATFORM_SWITCH
-		// checking the CanPlayOnline privilege on switch will log the user in if required in all but the NotLoggedIn state
+#if !PATCH_CHECK_PRIVILEGE_MUST_BE_LOGGED_IN
+		// some platforms will log the user in if required in all but the NotLoggedIn state
 		const bool bCanCheckPlayOnlinePrivilege = UserId.IsValid() && (PlatformOnlineIdentity->GetLoginStatus(*UserId) != ELoginStatus::NotLoggedIn);
 #else
 		const bool bCanCheckPlayOnlinePrivilege = UserId.IsValid() && (PlatformOnlineIdentity->GetLoginStatus(*UserId) == ELoginStatus::LoggedIn);
