Added UCurveTable modification callback deferral in OnlineHotfixManager similar to what we do for UDataTable, significantly saving the number of callbacks if we're doing a large number of row updates.

diff --git a/Engine/Plugins/Online/OnlineFramework/Source/Hotfix/Private/OnlineHotfixManager.cpp b/Engine/Plugins/Online/OnlineFramework/Source/Hotfix/Private/OnlineHotfixManager.cpp
--- a/Engine/Plugins/Online/OnlineFramework/Source/Hotfix/Private/OnlineHotfixManager.cpp
+++ b/Engine/Plugins/Online/OnlineFramework/Source/Hotfix/Private/OnlineHotfixManager.cpp
@@ -6,6 +6,8 @@
 #include "UObject/UObjectIterator.h"
 #include "UObject/Package.h"
 
+#include "HAL/ConsoleManager.h"
+
 #include "Misc/PackageName.h"
 #include "Misc/EngineVersion.h"
 #include "Misc/FileHelper.h"
@@ -52,6 +54,16 @@ public:
 	TArray<FString> Files;
 };
 
+namespace OnlineHotfixManagerCVars
+{
+	static bool bDeferBroadcastCurveTableModified = true;
+	static FAutoConsoleVariableRef DeferBroadcastCurveTableModifiedCVar(
+		TEXT("hotfix.DeferBroadcastCurveTableModified"),
+		bDeferBroadcastCurveTableModified,
+		TEXT("Whether to wait until all asset hotfixes have been applied before broadcasting OnCurveTableChanged delegates, as opposed to broadcasting after each individual modification"),
+		ECVF_Default);
+}
+
 namespace
 {
 	/** @return the expected network version for hotfix files determined at compile time */
@@ -1300,7 +1312,9 @@ void UOnlineHotfixManager::PatchAssetsFromIniFiles()
 			UCurveLinearColor::StaticClass(),
 		};
 
-		TSet<UDataTable*> ChangedTables;
+		TSet<UDataTable*> ChangedDataTables;
+		TSet<UCurveTable*> ChangedCurveTables;
+		TSet<UCurveTable*>* ChangedCurveTablesPointer = OnlineHotfixManagerCVars::bDeferBroadcastCurveTableModified ? &ChangedCurveTables : nullptr;
 
 		for (FConfigSection::TConstIterator It(*AssetHotfixConfigSection); It; ++It)
 		{
@@ -1361,7 +1375,7 @@ void UOnlineHotfixManager::PatchAssetsFromIniFiles()
 								//	+DataTable=<data table path>;RowUpdate;<row name>;<column name>;<new value>
 								//	+CurveTable=<curve table path>;RowUpdate;<row name>;<column name>;<new value>
 								//	+CurveFloat=<curve float path>;RowUpdate;None;<column name>;<new value>
-								HotfixRowUpdate(Asset, AssetPath, Tokens[2], Tokens[3], Tokens[4], ProblemStrings, &ChangedTables);
+								HotfixRowUpdate(Asset, AssetPath, Tokens[2], Tokens[3], Tokens[4], ProblemStrings, &ChangedDataTables, ChangedCurveTablesPointer);
 								bAddAssetToHotfixedList = ProblemStrings.Num() == 0;
 							}
 							else if ((HotfixType == TableUpdate || HotfixType == CurveUpdate) && Tokens.Num() == 3)
@@ -1434,11 +1448,19 @@ void UOnlineHotfixManager::PatchAssetsFromIniFiles()
 			}
 		}
 
-		for (UDataTable* Table : ChangedTables)
+		for (UDataTable* DataTable : ChangedDataTables)
 		{
-			if (Table != nullptr)
+			if (DataTable != nullptr)
 			{
-				Table->HandleDataTableChanged();
+				DataTable->HandleDataTableChanged();
+			}
+		}
+
+		for (UCurveTable* CurveTable : ChangedCurveTables)
+		{
+			if (CurveTable != nullptr)
+			{
+				CurveTable->OnCurveTableChanged().Broadcast();
 			}
 		}
 	}
@@ -1519,7 +1541,15 @@ void UOnlineHotfixManager::ReloadConfigsFromIniFiles()
 		ClassesToReload.Num(), FPlatformTime::Seconds() - StartTime, NumObjectsReloaded);
 }
 
-void UOnlineHotfixManager::HotfixRowUpdate(UObject* Asset, const FString& AssetPath, const FString& RowName, const FString& ColumnName, const FString& NewValue, TArray<FString>& ProblemStrings, TSet<UDataTable*>* ChangedTables)
+void UOnlineHotfixManager::HotfixRowUpdate(
+	UObject* Asset,
+	const FString& AssetPath,
+	const FString& RowName,
+	const FString& ColumnName,
+	const FString& NewValue,
+	TArray<FString>& ProblemStrings,
+	TSet<UDataTable*>* ChangedDataTables,
+	TSet<UCurveTable*>* ChangedCurveTables)
 {
 	if (AssetPath.IsEmpty())
 	{
@@ -1672,13 +1702,13 @@ void UOnlineHotfixManager::HotfixRowUpdate(UObject* Asset, const FString& AssetP
 
 		if (bWasDataTableChanged)
 		{
-			if (ChangedTables == nullptr)
+			if (ChangedDataTables == nullptr)
 			{
 				DataTable->HandleDataTableChanged();
 			}
 			else
 			{
-				ChangedTables->Add(DataTable);
+				ChangedDataTables->Add(DataTable);
 			}
 		}
 	}
@@ -1745,7 +1775,14 @@ void UOnlineHotfixManager::HotfixRowUpdate(UObject* Asset, const FString& AssetP
 
 		if (bWasCurveTableChanged)
 		{
-			CurveTable->OnCurveTableChanged().Broadcast();
+			if (ChangedCurveTables == nullptr)
+			{
+				CurveTable->OnCurveTableChanged().Broadcast();
+			}
+			else
+			{
+				ChangedCurveTables->Add(CurveTable);
+			}
 		}
 	}
 	else if (CurveFloat)
diff --git a/Engine/Plugins/Online/OnlineFramework/Source/Hotfix/Public/OnlineHotfixManager.h b/Engine/Plugins/Online/OnlineFramework/Source/Hotfix/Public/OnlineHotfixManager.h
--- a/Engine/Plugins/Online/OnlineFramework/Source/Hotfix/Public/OnlineHotfixManager.h
+++ b/Engine/Plugins/Online/OnlineFramework/Source/Hotfix/Public/OnlineHotfixManager.h
@@ -353,7 +353,15 @@ protected:
 	/** Used in PatchAssetsFromIniFiles to hotfix only a row in a table. 
 	 *  If ChangedTables is not null then HandleDataTableChanged will not be called and the caller should call it on the data tables in ChangedTables when they're ready to
 	 */
-	void HotfixRowUpdate(UObject* Asset, const FString& AssetPath, const FString& RowName, const FString& ColumnName, const FString& NewValue, TArray<FString>& ProblemStrings, TSet<class UDataTable*>* ChangedTables = nullptr);
+	void HotfixRowUpdate(
+		UObject* Asset,
+		const FString& AssetPath,
+		const FString& RowName,
+		const FString& ColumnName,
+		const FString& NewValue,
+		TArray<FString>& ProblemStrings,
+		TSet<class UDataTable*>* ChangedDataTables = nullptr,
+		TSet<class UCurveTable*>* ChangedCurveTables = nullptr);
 	
 	/** Used in PatchAssetsFromIniFiles to hotfix an entire table. */
 	void HotfixTableUpdate(UObject* Asset, const FString& AssetPath, const FString& JsonData, TArray<FString>& ProblemStrings);
