Add scope counters around functions found in hitch jiras

diff --git a/Engine/Plugins/Online/OnlineFramework/Source/Party/Private/Party/SocialParty.cpp b/Engine/Plugins/Online/OnlineFramework/Source/Party/Private/Party/SocialParty.cpp
--- a/Engine/Plugins/Online/OnlineFramework/Source/Party/Private/Party/SocialParty.cpp
+++ b/Engine/Plugins/Online/OnlineFramework/Source/Party/Private/Party/SocialParty.cpp
@@ -1785,6 +1785,7 @@ void USocialParty::DisconnectParty()
 
 void USocialParty::FinalizePartyLeave(EMemberExitedReason Reason)
 {
+	QUICK_SCOPE_CYCLE_COUNTER(STAT_USocialParty_FinalizePartyLeave);
 	UE_LOG(LogParty, Verbose, TEXT("Local player [%s] is no longer in party [%s]. Reason [%s]."), *GetOwningLocalMember().ToDebugString(false), *ToDebugString(), ToString(Reason));
 
 	if (!bIsLeavingParty)
