Fix the OnlineHotFixManager so it doesn't wait for all async load requests to complete unless it has to.

PIE time is reduced locally by ~5-10s since we rarely have hotfixes to apply when developing and thus have no reason to force a stall of async loads

diff --git a/Engine/Plugins/Online/OnlineFramework/Source/Hotfix/Private/OnlineHotfixManager.cpp b/Engine/Plugins/Online/OnlineFramework/Source/Hotfix/Private/OnlineHotfixManager.cpp
--- a/Engine/Plugins/Online/OnlineFramework/Source/Hotfix/Private/OnlineHotfixManager.cpp
+++ b/Engine/Plugins/Online/OnlineFramework/Source/Hotfix/Private/OnlineHotfixManager.cpp
@@ -1292,9 +1292,6 @@ void UOnlineHotfixManager::PatchAssetsFromIniFiles()
 {
 	UE_LOG(LogHotfixManager, Display, TEXT("Checking for assets to be patched using data from 'AssetHotfix' section in the Game .ini file"));
 
-	// Flush async loading before modifying GConfig.
-	FlushAsyncLoading();
-
 	int32 TotalPatchableAssets = 0;
 	AssetsHotfixedFromIniFiles.Reset();
 
@@ -1302,6 +1299,9 @@ void UOnlineHotfixManager::PatchAssetsFromIniFiles()
 	const FConfigSection* AssetHotfixConfigSection = GConfig->GetSection(TEXT("AssetHotfix"), false, GGameIni);
 	if (AssetHotfixConfigSection != nullptr)
 	{
+		// Flush async loading before modifying GConfig.
+		FlushAsyncLoading();
+
 		// These are the asset types we support patching right now
 		UClass* const PatchableAssetClasses[] = 
 		{ 
@@ -1481,6 +1481,11 @@ void UOnlineHotfixManager::PatchAssetsFromIniFiles()
 
 void UOnlineHotfixManager::ReloadConfigsFromIniFiles()
 {
+	if (HotfixFileList.IsEmpty())
+	{
+		return;
+	}
+
 	FlushAsyncLoading();
 
 	TArray<FString> ClassesToReload;
