Remove unnecessary include

diff --git a/Engine/Plugins/Online/OnlineFramework/Source/Party/Public/SocialManager.h b/Engine/Plugins/Online/OnlineFramework/Source/Party/Public/SocialManager.h
--- a/Engine/Plugins/Online/OnlineFramework/Source/Party/Public/SocialManager.h
+++ b/Engine/Plugins/Online/OnlineFramework/Source/Party/Public/SocialManager.h
@@ -18,7 +18,6 @@ class ULocalPlayer;
 class USocialUser;
 class USocialParty;
 class USocialToolkit;
-class UGameViewportClient;
 class UGameInstance;
 class FOnlineSessionSearchResult;
 class FPartyPlatformSessionManager;
