Add analytics for QoS evaluation of accelerated routes.

diff --git a/Engine/Plugins/Online/OnlineFramework/Source/Qos/Private/QosEvaluator.cpp b/Engine/Plugins/Online/OnlineFramework/Source/Qos/Private/QosEvaluator.cpp
--- a/Engine/Plugins/Online/OnlineFramework/Source/Qos/Private/QosEvaluator.cpp
+++ b/Engine/Plugins/Online/OnlineFramework/Source/Qos/Private/QosEvaluator.cpp
@@ -85,7 +85,7 @@ void UQosEvaluator::FindDatacenters(const FQosParams& InParams, const TArray<FQo
 		UE_LOG(LogQos, Log, TEXT("Qos evaluation already in progress, ignoring"));
 		// Just trigger delegate now (Finalize resets state vars)
 		GetWorldTimerManager().SetTimerForNextTick(FTimerDelegate::CreateLambda([InCompletionDelegate]() {
-			InCompletionDelegate.ExecuteIfBound(EQosCompletionResult::Failure, TArray<FDatacenterQosInstance>());
+			InCompletionDelegate.ExecuteIfBound(EQosCompletionResult::Failure, TArray<FDatacenterQosInstance>(), nullptr, nullptr);
 		}));
 	}
 }
@@ -206,8 +206,16 @@ void UQosEvaluator::OnEchoManyCompleted(FIcmpEchoManyCompleteResult FinalResult,
 
 		CompletionResult = EQosCompletionResult::Success;
 		CalculatePingAverages();
+		FString SelectedSubRegion;
+		FString SelectedRegion;
+		InQosSearchCompleteDelegate.ExecuteIfBound(CompletionResult, Datacenters, &SelectedRegion, &SelectedSubRegion);
+
+		if (QosStats.IsValid())
+		{
+			QosStats->SetChosenRegion(MoveTemp(SelectedRegion));
+			QosStats->SetChosenSubRegion(MoveTemp(SelectedSubRegion));
+		}
 		EndAnalytics(CompletionResult);
-		InQosSearchCompleteDelegate.ExecuteIfBound(CompletionResult, Datacenters);
 	}
 	break;
 
@@ -217,7 +225,7 @@ void UQosEvaluator::OnEchoManyCompleted(FIcmpEchoManyCompleteResult FinalResult,
 
 		CompletionResult = EQosCompletionResult::Failure;
 		EndAnalytics(CompletionResult);
-		InQosSearchCompleteDelegate.ExecuteIfBound(CompletionResult, Datacenters);
+		InQosSearchCompleteDelegate.ExecuteIfBound(CompletionResult, Datacenters, nullptr, nullptr);
 	}
 	break;
 
@@ -227,7 +235,7 @@ void UQosEvaluator::OnEchoManyCompleted(FIcmpEchoManyCompleteResult FinalResult,
 
 		CompletionResult = EQosCompletionResult::Canceled;
 		EndAnalytics(CompletionResult);
-		InQosSearchCompleteDelegate.ExecuteIfBound(CompletionResult, Datacenters);
+		InQosSearchCompleteDelegate.ExecuteIfBound(CompletionResult, Datacenters, nullptr, nullptr);
 	}
 	break;
 
@@ -259,7 +267,6 @@ void UQosEvaluator::EndAnalytics(EQosCompletionResult CompletionResult)
 			{
 				ResultType = EDatacenterResultType::Normal;
 			}
-
 			QosStats->EndQosPass(ResultType);
 			QosStats->Upload(AnalyticsProvider);
 		}
@@ -358,7 +365,7 @@ bool UQosEvaluator::PingRegionServers(const FQosParams& InParams, const FOnQosSe
 	if (0 == Targets.Num())
 	{
 		// Nothing to do if no servers provided in the list of datacenters.
-		InQosSearchCompleteDelegate.ExecuteIfBound(EQosCompletionResult::Failure, Datacenters);
+		InQosSearchCompleteDelegate.ExecuteIfBound(EQosCompletionResult::Failure, Datacenters, nullptr, nullptr);
 		return false;
 	}
 
diff --git a/Engine/Plugins/Online/OnlineFramework/Source/Qos/Private/QosEvaluator.h b/Engine/Plugins/Online/OnlineFramework/Source/Qos/Private/QosEvaluator.h
--- a/Engine/Plugins/Online/OnlineFramework/Source/Qos/Private/QosEvaluator.h
+++ b/Engine/Plugins/Online/OnlineFramework/Source/Qos/Private/QosEvaluator.h
@@ -36,8 +36,10 @@ DECLARE_DELEGATE_OneParam(FOnQosPingEvalComplete, EQosCompletionResult /** Resul
  *
  * @param Result the QoS operation result
  * @param DatacenterInstances The per-datacenter ping information
+ * @param OutSelectedRegion If present target will be set to the ReionManagers recommended region (may differ from actual best)
+ * @param OutSelectedSubRegion If present target will be set to the ReionManagers recommended sub region (may differ from actual best)
  */
-DECLARE_DELEGATE_TwoParams(FOnQosSearchComplete, EQosCompletionResult /** Result */, const TArray<FDatacenterQosInstance>& /** DatacenterInstances */);
+DECLARE_DELEGATE_FourParams(FOnQosSearchComplete, EQosCompletionResult /** Result */, const TArray<FDatacenterQosInstance>& /** DatacenterInstances */, FString*  /** OutSelectedRegion */, FString*  /** OutSelectedSubRegion */);
 
 
 /**
diff --git a/Engine/Plugins/Online/OnlineFramework/Source/Qos/Private/QosRegionManager.cpp b/Engine/Plugins/Online/OnlineFramework/Source/Qos/Private/QosRegionManager.cpp
--- a/Engine/Plugins/Online/OnlineFramework/Source/Qos/Private/QosRegionManager.cpp
+++ b/Engine/Plugins/Online/OnlineFramework/Source/Qos/Private/QosRegionManager.cpp
@@ -401,7 +401,7 @@ bool UQosRegionManager::IsQosEvaluationInProgress() const
 	return Evaluator != nullptr;
 }
 
-void UQosRegionManager::OnQosEvaluationComplete(EQosCompletionResult Result, const TArray<FDatacenterQosInstance>& DatacenterInstances)
+void UQosRegionManager::OnQosEvaluationComplete(EQosCompletionResult Result, const TArray<FDatacenterQosInstance>& DatacenterInstances, FString* OutSelectedRegion, FString* OutSelectedSubRegion)
 {
 	// toss the evaluator
 	if (Evaluator != nullptr)
@@ -498,10 +498,21 @@ void UQosRegionManager::OnQosEvaluationComplete(EQosCompletionResult Result, con
 			TrySetDefaultRegion();
 		}
 	}
+	FString BestRegion = GetBestRegion();
 	TArray<FString> BestRegionSubregions;
-	GetSubregionPreferences(GetBestRegion(), BestRegionSubregions);
+	GetSubregionPreferences(BestRegion, BestRegionSubregions);
 	UE_LOG(LogQos, Log, TEXT("[UQosRegionManager::OnQosEvaluationComplete] ping eval has completed.  Best region is '%s', recommended subregion is '%s'"),
-		*GetBestRegion(), BestRegionSubregions.Num() ? *BestRegionSubregions[0] : TEXT("NONE"));
+		*BestRegion, BestRegionSubregions.Num() ? *BestRegionSubregions[0] : TEXT("NONE"));
+
+	if (OutSelectedRegion)
+	{
+		*OutSelectedRegion = BestRegion;
+	}
+	
+	if (OutSelectedSubRegion)
+	{
+		*OutSelectedSubRegion = *BestRegionSubregions[0];
+	}
 
 #if DEBUG_SUBCOMPARE_BY_SUBSPACE
 	DumpRegionStats();
diff --git a/Engine/Plugins/Online/OnlineFramework/Source/Qos/Private/QosStats.cpp b/Engine/Plugins/Online/OnlineFramework/Source/Qos/Private/QosStats.cpp
--- a/Engine/Plugins/Online/OnlineFramework/Source/Qos/Private/QosStats.cpp
+++ b/Engine/Plugins/Online/OnlineFramework/Source/Qos/Private/QosStats.cpp
@@ -28,6 +28,9 @@ const FString FQosDatacenterStats::QosStats_BestRegionId = TEXT("BestRegionId");
 const FString FQosDatacenterStats::QosStats_BestRegionPing = TEXT("BestRegionPing");
 const FString FQosDatacenterStats::QosStats_BestEndpointId = TEXT("BestEndpointId");
 const FString FQosDatacenterStats::QosStats_BestEndpointPing = TEXT("BestEndpointPing");
+const FString FQosDatacenterStats::QosStats_ChosenRegionId = TEXT("ChosenbRegionId");
+const FString FQosDatacenterStats::QosStats_ChosenSubRegionId = TEXT("ChosenSubRegionId");
+const FString FQosDatacenterStats::QosStats_ChosenSubRegionPing = TEXT("ChosenSubRegionPing");
 
 /**
  * Debug output for the contents of a recorded stats event
@@ -108,6 +111,22 @@ void FQosDatacenterStats::RecordQosAttempt(const FString& Region, const FString&
 	}
 }
 
+void FQosDatacenterStats::SetChosenSubRegion(FString&& SubRegion)
+{
+	if (bAnalyticsInProgress)
+	{
+		QosData.RulesChosenSubRegion = MoveTemp(SubRegion);
+	}
+}
+
+void FQosDatacenterStats::SetChosenRegion(FString && Region)
+{
+	if (bAnalyticsInProgress)
+	{
+		QosData.RulesChosenRegion = MoveTemp(Region);
+	}
+}
+
 void FQosDatacenterStats::EndQosPass(EDatacenterResultType Result)
 {
 	if (bAnalyticsInProgress)
@@ -154,6 +173,8 @@ void FQosDatacenterStats::Upload(TSharedPtr<IAnalyticsProvider>& AnalyticsProvid
  * @EventParam BestRegionPing integer ping in the best RegionId (that is usable, 0 if BestRegionId is UNREACHABLE)
  * @EventParam BestEndpointId string RegionId with best ping (regardless of usability, UNREACHABLE if none pass QoS)
  * @EventParam BestEndpointPing integer ping in the best RegionId (regardless of usability, 0 if BestEndpointId is UNREACHABLE)
+ * @EventParam ChosenRegionId string RegionId chosen by RegionManager including bias, may differ form best (that is usable, UNREACHABLE if none pass QoS)
+ * @EventParam ChosenRegionPing integer ping in the chosen RegionId (that is usable, 0 if BestRegionId is UNREACHABLE)
  * @EventParam RegionDetails json representation of ping details
  * @Comments Analytics data for a complete qos datacenter determination attempt
  * 
@@ -179,6 +200,9 @@ void FQosDatacenterStats::ParseQosResults(TSharedPtr<IAnalyticsProvider>& Analyt
 		FString BestEndpointId(TEXT("Unknown"));
 		int32 BestEndpointPing = INT_MAX;
 
+		FString ChosenSubRegionId(TEXT("Unknown"));
+		int32 ChosenSubRegionPing = INT_MAX;
+
 		for (const FQosStats_RegionInfo& Region : QosData.Regions)
 		{
 			if (Region.AvgPing < BestPing)
@@ -192,9 +216,16 @@ void FQosDatacenterStats::ParseQosResults(TSharedPtr<IAnalyticsProvider>& Analyt
 				BestEndpointId = Region.RegionId;
 				BestEndpointPing = Region.AvgPing;
 			}
+
+			if(Region.RegionId == QosData.RulesChosenSubRegion)
+			{
+				ChosenSubRegionId = Region.RegionId;
+				ChosenSubRegionPing = Region.AvgPing;
+			}
 		}
 
 		BestPing = FMath::Clamp(BestPing, 0, UNREACHABLE_PING);
+		ChosenSubRegionPing = FMath::Clamp(ChosenSubRegionPing, 0, UNREACHABLE_PING);
 		if (BestPing < UNREACHABLE_PING)
 		{
 			QoSAttributes.Add(FAnalyticsEventAttribute(QosStats_BestRegionId, BestRegionId));
@@ -206,6 +237,22 @@ void FQosDatacenterStats::ParseQosResults(TSharedPtr<IAnalyticsProvider>& Analyt
 			QoSAttributes.Add(FAnalyticsEventAttribute(QosStats_BestRegionPing, 0));
 		}
 
+		if (ChosenSubRegionPing < UNREACHABLE_PING)
+		{
+			QoSAttributes.Add(FAnalyticsEventAttribute(QosStats_ChosenSubRegionId, ChosenSubRegionId));
+			QoSAttributes.Add(FAnalyticsEventAttribute(QosStats_ChosenSubRegionPing, ChosenSubRegionPing));
+		}
+		else
+		{
+			QoSAttributes.Add(FAnalyticsEventAttribute(QosStats_ChosenSubRegionId, TEXT("UNREACHABLE")));
+			QoSAttributes.Add(FAnalyticsEventAttribute(QosStats_ChosenSubRegionPing, 0));
+		}
+
+		if(!QosData.RulesChosenRegion.IsEmpty())
+		{
+			QoSAttributes.Add(FAnalyticsEventAttribute(QosStats_ChosenRegionId, QosData.RulesChosenRegion));
+		}
+
 		BestEndpointPing = FMath::Clamp(BestEndpointPing, 0, UNREACHABLE_PING);
 		if (BestEndpointPing < UNREACHABLE_PING)
 		{
diff --git a/Engine/Plugins/Online/OnlineFramework/Source/Qos/Private/QosStats.h b/Engine/Plugins/Online/OnlineFramework/Source/Qos/Private/QosStats.h
--- a/Engine/Plugins/Online/OnlineFramework/Source/Qos/Private/QosStats.h
+++ b/Engine/Plugins/Online/OnlineFramework/Source/Qos/Private/QosStats.h
@@ -133,6 +133,10 @@ private:
 		FQosStats_Timer SearchTime;
 		/** Array of region information */
 		TArray<FQosStats_RegionInfo> Regions;
+		/** name for the region chosen based on rules */
+		FString  RulesChosenRegion;
+		/** name for the sub region chosen based on rules */
+		FString  RulesChosenSubRegion;
 		/** Number of search results tested */
 		int32 NumTotalSearches;
 		/** Number of successful search results */
@@ -170,6 +174,9 @@ private:
 	static const FString QosStats_BestRegionPing;
 	static const FString QosStats_BestEndpointId;
 	static const FString QosStats_BestEndpointPing;
+	static const FString QosStats_ChosenRegionId;
+	static const FString QosStats_ChosenSubRegionId;
+	static const FString QosStats_ChosenSubRegionPing;
 
 	/** Version of the stats for separation */
 	int32 StatsVersion;
@@ -242,6 +249,20 @@ public:
 	 */
 	void EndQosPass(EDatacenterResultType Result);
 
+	/**
+	 * Set the chosen region from Qos Evaluation (may differ from the fastest/best based on rules)
+	 *
+	 * @param Region string region code of chosen region
+	 */
+	void SetChosenRegion(FString&& Region);
+
+	/**
+	 * Set the chosen SubRegion from Qos Evaluation (may differ from the fastest/best based on rules)
+	 *
+	 * @param Region string SubRegion code of chosen region
+	 */
+	void SetChosenSubRegion(FString && SubRegion);
+
 	/**
 	 * Record previously saved stats to an analytics provider
 	 *
diff --git a/Engine/Plugins/Online/OnlineFramework/Source/Qos/Public/QosRegionManager.h b/Engine/Plugins/Online/OnlineFramework/Source/Qos/Public/QosRegionManager.h
--- a/Engine/Plugins/Online/OnlineFramework/Source/Qos/Public/QosRegionManager.h
+++ b/Engine/Plugins/Online/OnlineFramework/Source/Qos/Public/QosRegionManager.h
@@ -602,7 +602,7 @@ private:
 	 */
 	void SanityCheckDefinitions() const;
 
-	void OnQosEvaluationComplete(EQosCompletionResult Result, const TArray<FDatacenterQosInstance>& DatacenterInstances);
+	void OnQosEvaluationComplete(EQosCompletionResult Result, const TArray<FDatacenterQosInstance>& DatacenterInstances, FString* OutSelectedRegion, FString* OutSelectedSubRegion);
 
 	/**
 	 * Use the existing set value, or if it is currently invalid, set the next best region available
