Add USocialPrivateMessageChannel::GetTargetUser() public getter

diff --git a/Engine/Plugins/Online/OnlineFramework/Source/Party/Private/Chat/SocialPrivateMessageChannel.cpp b/Engine/Plugins/Online/OnlineFramework/Source/Party/Private/Chat/SocialPrivateMessageChannel.cpp
--- a/Engine/Plugins/Online/OnlineFramework/Source/Party/Private/Chat/SocialPrivateMessageChannel.cpp
+++ b/Engine/Plugins/Online/OnlineFramework/Source/Party/Private/Chat/SocialPrivateMessageChannel.cpp
@@ -49,3 +49,8 @@ void USocialPrivateMessageChannel::SetTargetUser(USocialUser& InTargetUser)
 	TargetUser = &InTargetUser;
 }
 
+USocialUser& USocialPrivateMessageChannel::GetTargetUser() const
+{
+	check(TargetUser);
+	return *TargetUser;
+}
diff --git a/Engine/Plugins/Online/OnlineFramework/Source/Party/Public/Chat/SocialPrivateMessageChannel.h b/Engine/Plugins/Online/OnlineFramework/Source/Party/Public/Chat/SocialPrivateMessageChannel.h
--- a/Engine/Plugins/Online/OnlineFramework/Source/Party/Public/Chat/SocialPrivateMessageChannel.h
+++ b/Engine/Plugins/Online/OnlineFramework/Source/Party/Public/Chat/SocialPrivateMessageChannel.h
@@ -18,6 +18,8 @@ public:
 	virtual void Initialize(USocialUser* InSocialUser, const FChatRoomId& InChannelId, ESocialChannelType InSourceChannelType) override;
 
 	virtual bool SendMessage(const FString& InMessage) override;
+	
+	USocialUser& GetTargetUser() const;
 
 private:
 	void SetTargetUser(USocialUser& InTargetUser);
