Made party kick evaluation and execution functions virutal so they may be handled differently at the game layer.

#tests Compiled DevelopmentEditor Win64 / Ran a PIE instance with a second local player and verified kick flow.
#rnx

diff --git a/Engine/Plugins/Online/OnlineFramework/Source/Party/Public/Party/SocialParty.h b/Engine/Plugins/Online/OnlineFramework/Source/Party/Public/Party/SocialParty.h
--- a/Engine/Plugins/Online/OnlineFramework/Source/Party/Public/Party/SocialParty.h
+++ b/Engine/Plugins/Online/OnlineFramework/Source/Party/Public/Party/SocialParty.h
@@ -237,11 +237,11 @@ protected:
 	bool HasUserBeenInvited(const USocialUser& User) const;
 	
 	bool CanPromoteMember(const UPartyMember& PartyMember) const;
-	bool CanKickMember(const UPartyMember& PartyMember) const;
+	virtual bool CanKickMember(const UPartyMember& PartyMember) const;
 	
 	bool TryInviteUser(const USocialUser& UserToInvite, const ESocialPartyInviteMethod InviteMethod = ESocialPartyInviteMethod::Other, const FString& MetaData = FString());
 	bool TryPromoteMember(const UPartyMember& PartyMember);
-	bool TryKickMember(const UPartyMember& PartyMember);
+	virtual bool TryKickMember(const UPartyMember& PartyMember);
 	//--------------------------
 
 protected:
