Fix USocialParty::ContainsUser to look at Primary net IDs instead of comparing SocialUsers. There can be multiple SocialUsers with the same underlying Primary ID depending on where the SocialUser originated from (such as a different Toolkit).

This error caused FSocialUserList::EvaluateUserPresence to not detect users in the same party.

#tests Compiled DevelopmentEditor Win64 / Ran a PIE instance and added a second local player. Verified the players showed up in each others Party Members user list header.
#rnx

diff --git a/Engine/Plugins/Online/OnlineFramework/Source/Party/Private/Party/SocialParty.cpp b/Engine/Plugins/Online/OnlineFramework/Source/Party/Private/Party/SocialParty.cpp
--- a/Engine/Plugins/Online/OnlineFramework/Source/Party/Private/Party/SocialParty.cpp
+++ b/Engine/Plugins/Online/OnlineFramework/Source/Party/Private/Party/SocialParty.cpp
@@ -1368,7 +1368,7 @@ bool USocialParty::ContainsUser(const USocialUser& User) const
 {
 	for(const UPartyMember* PartyMember : GetPartyMembers())
 	{
-		if (&PartyMember->GetSocialUser() == &User)
+		if (PartyMember->GetSocialUser().GetUserId(ESocialSubsystem::Primary) == User.GetUserId(ESocialSubsystem::Primary))
 		{
 			return true;
 		}
