Created new GetOwningLocalPlayerPtr method and deprecated GetOwningLocalPlayer. GetOwningLocalPlayerPtr returns a pointer instead of a reference as the underlying data is a TWeakObjectPtr. The weak LocalPlayer can be removed when the user is logged out.

#tests Compiled DevelopmentEditor Win64 / Test Client / Ran a PIE instance and logged into and out of several splitscreen clients. Opened and closed PIE several times.

diff --git a/Engine/Plugins/Online/OnlineFramework/Source/Party/Private/Interactions/PartyInteractions.cpp b/Engine/Plugins/Online/OnlineFramework/Source/Party/Private/Interactions/PartyInteractions.cpp
--- a/Engine/Plugins/Online/OnlineFramework/Source/Party/Private/Interactions/PartyInteractions.cpp
+++ b/Engine/Plugins/Online/OnlineFramework/Source/Party/Private/Interactions/PartyInteractions.cpp
@@ -197,9 +197,9 @@ bool FSocialInteraction_LeaveParty::CanExecute(const USocialUser& User)
 {
 	if (User.IsLocalUser())
 	{
-		USocialToolkit& OwningToolkit = User.GetOwningToolkit();
-		ULocalPlayer& LocalPlayer = OwningToolkit.GetOwningLocalPlayer();
-		if (!LocalPlayer.IsPrimaryPlayer())
+		const USocialToolkit& OwningToolkit = User.GetOwningToolkit();
+		const ULocalPlayer* LocalPlayer = OwningToolkit.GetOwningLocalPlayerPtr();
+		if (!LocalPlayer || !LocalPlayer->IsPrimaryPlayer())
 		{
 			return false;
 		}
diff --git a/Engine/Plugins/Online/OnlineFramework/Source/Party/Private/Party/SocialParty.cpp b/Engine/Plugins/Online/OnlineFramework/Source/Party/Private/Party/SocialParty.cpp
--- a/Engine/Plugins/Online/OnlineFramework/Source/Party/Private/Party/SocialParty.cpp
+++ b/Engine/Plugins/Online/OnlineFramework/Source/Party/Private/Party/SocialParty.cpp
@@ -354,7 +354,10 @@ bool USocialParty::TryInviteUser(const USocialUser& UserToInvite, const ESocialP
 				{
 					if (const IOnlineIdentityPtr PlatformIdentityInterface = Online::GetIdentityInterface(GetWorld(), SocialOssName))
 					{
-						LocalUserPlatformId = PlatformIdentityInterface->GetUniquePlayerId(GetOwningLocalPlayer().GetControllerId());
+						if (const ULocalPlayer* LocalPlayer = GetOwningLocalPlayerPtr())
+						{
+							LocalUserPlatformId = PlatformIdentityInterface->GetUniquePlayerId(LocalPlayer->GetControllerId());
+						}
 					}
 				}
 
@@ -1378,9 +1381,23 @@ bool USocialParty::ContainsUser(const USocialUser& User) const
 }
 
 ULocalPlayer& USocialParty::GetOwningLocalPlayer() const
+{
+	// This is deprecated as it was unsafe.
+	// The Toolkit's LocalPlayerOwner is a TWeakObjectPtr and may return nullptr when the local player logs out. Please use the pointer version.
+	if (ULocalPlayer* LocalPlayer = GetOwningLocalPlayerPtr())
+	{
+		return *LocalPlayer;
+	}
+	else
+	{
+		return *NewObject<ULocalPlayer>();
+	}
+}
+
+ULocalPlayer* USocialParty::GetOwningLocalPlayerPtr() const
 {
 	//@todo DanH Party: This is a wee bit heavy - should be able to do this in fewer steps
-	return GetOwningLocalMember().GetSocialUser().GetOwningToolkit().GetOwningLocalPlayer();
+	return GetOwningLocalMember().GetSocialUser().GetOwningToolkit().GetOwningLocalPlayerPtr();
 }
 
 bool USocialParty::IsLocalPlayerPartyLeader() const
@@ -1710,7 +1727,8 @@ void USocialParty::CleanupSpectatorBeacon()
 
 FName USocialParty::GetGameSessionName() const
 {
-	const APlayerController* OwnerPC = GetOwningLocalPlayer().GetPlayerController(GetWorld());
+	const ULocalPlayer* LocalPlayer = GetOwningLocalPlayerPtr();
+	const APlayerController* OwnerPC = LocalPlayer ? LocalPlayer->GetPlayerController(GetWorld()) : nullptr;
 	if (OwnerPC && OwnerPC->PlayerState)
 	{
 		return OwnerPC->PlayerState->SessionName;
diff --git a/Engine/Plugins/Online/OnlineFramework/Source/Party/Private/SocialManager.cpp b/Engine/Plugins/Online/OnlineFramework/Source/Party/Private/SocialManager.cpp
--- a/Engine/Plugins/Online/OnlineFramework/Source/Party/Private/SocialManager.cpp
+++ b/Engine/Plugins/Online/OnlineFramework/Source/Party/Private/SocialManager.cpp
@@ -266,7 +266,7 @@ USocialToolkit& USocialManager::GetSocialToolkit(const ULocalPlayer& LocalPlayer
 	USocialToolkit* FoundToolkit = nullptr;
 	for (USocialToolkit* Toolkit : SocialToolkits)
 	{
-		if (&LocalPlayer == &Toolkit->GetOwningLocalPlayer())
+		if (&LocalPlayer == Toolkit->GetOwningLocalPlayerPtr())
 		{
 			FoundToolkit = Toolkit;
 			break;
@@ -293,7 +293,8 @@ USocialToolkit* USocialManager::GetSocialToolkit(FUniqueNetIdRepl LocalUserId) c
 {
 	for (USocialToolkit* Toolkit : SocialToolkits)
 	{
-		if (Toolkit->GetOwningLocalPlayer().GetPreferredUniqueNetId() == LocalUserId)
+		const ULocalPlayer* LocalPlayer = Toolkit->GetOwningLocalPlayerPtr();
+		if (LocalPlayer && LocalPlayer->GetPreferredUniqueNetId() == LocalUserId)
 		{
 			return Toolkit;
 		}
@@ -740,7 +741,7 @@ USocialToolkit& USocialManager::CreateSocialToolkit(ULocalPlayer& OwningLocalPla
 {
 	for (USocialToolkit* ExistingToolkit : SocialToolkits)
 	{
-		check(&OwningLocalPlayer != &ExistingToolkit->GetOwningLocalPlayer());
+		check(&OwningLocalPlayer != ExistingToolkit->GetOwningLocalPlayerPtr());
 	}
 	check(ToolkitClass);
 
diff --git a/Engine/Plugins/Online/OnlineFramework/Source/Party/Private/SocialToolkit.cpp b/Engine/Plugins/Online/OnlineFramework/Source/Party/Private/SocialToolkit.cpp
--- a/Engine/Plugins/Online/OnlineFramework/Source/Party/Private/SocialToolkit.cpp
+++ b/Engine/Plugins/Online/OnlineFramework/Source/Party/Private/SocialToolkit.cpp
@@ -248,7 +248,16 @@ FUniqueNetIdRepl USocialToolkit::GetLocalUserNetId(ESocialSubsystem SubsystemTyp
 
 int32 USocialToolkit::GetLocalUserNum() const
 {
-	return GetOwningLocalPlayer().GetControllerId();
+	if (const ULocalPlayer* LocalPlayer = GetOwningLocalPlayerPtr())
+	{
+		return LocalPlayer->GetControllerId();
+	}
+	else
+	{
+		UE_LOG(LogParty, Verbose, TEXT("%hs - OwningLocalPlayer is nullptr. Likely the player has logged out."), __FUNCTION__);
+
+		return INVALID_CONTROLLERID;
+	}
 }
 
 const FOnlineUserPresence* USocialToolkit::GetPresenceInfo(ESocialSubsystem SubsystemType) const
@@ -320,8 +329,21 @@ USocialManager& USocialToolkit::GetSocialManager() const
 
 ULocalPlayer& USocialToolkit::GetOwningLocalPlayer() const
 {
-	check(LocalPlayerOwner.IsValid());
-	return *LocalPlayerOwner.Get();
+	// This is deprecated as it is unsafe.
+	// LocalPlayerOwner is a TWeakObjectPtr and may return nullptr when the local player logs out. Please use the pointer version.
+	if (ULocalPlayer* LocalPlayer = GetOwningLocalPlayerPtr())
+	{
+		return *LocalPlayer;
+	}
+	else
+	{
+		return *NewObject<ULocalPlayer>();
+	}
+}
+
+ULocalPlayer* USocialToolkit::GetOwningLocalPlayerPtr() const
+{
+	return LocalPlayerOwner.Get();
 }
 
 USocialUser* USocialToolkit::FindUser(const FUniqueNetIdRepl& UserId) const
diff --git a/Engine/Plugins/Online/OnlineFramework/Source/Party/Public/Party/SocialParty.h b/Engine/Plugins/Online/OnlineFramework/Source/Party/Public/Party/SocialParty.h
--- a/Engine/Plugins/Online/OnlineFramework/Source/Party/Public/Party/SocialParty.h
+++ b/Engine/Plugins/Online/OnlineFramework/Source/Party/Public/Party/SocialParty.h
@@ -129,7 +129,10 @@ public:
 
 	bool ContainsUser(const USocialUser& User) const;
 
+	UE_DEPRECATED(5.6, "GetOwningLocalPlayer returns the Toolkit's LocalPlayerOwner which is a TWeakObjectPtr and may return nullptr when the local player logs out. Please use the pointer version.")
 	ULocalPlayer& GetOwningLocalPlayer() const;
+
+	ULocalPlayer* GetOwningLocalPlayerPtr() const;
 	const FUniqueNetIdRepl& GetOwningLocalUserId() const { return OwningLocalUserId; }
 	const FUniqueNetIdRepl& GetPartyLeaderId() const { return CurrentLeaderId; }
 	bool IsLocalPlayerPartyLeader() const;
diff --git a/Engine/Plugins/Online/OnlineFramework/Source/Party/Public/SocialToolkit.h b/Engine/Plugins/Online/OnlineFramework/Source/Party/Public/SocialToolkit.h
--- a/Engine/Plugins/Online/OnlineFramework/Source/Party/Public/SocialToolkit.h
+++ b/Engine/Plugins/Online/OnlineFramework/Source/Party/Public/SocialToolkit.h
@@ -78,8 +78,12 @@ public:
 	void SetLocalUserOnlineState(EOnlinePresenceState::Type OnlineState);
 	void AddLocalUserOnlineProperties(FPresenceProperties OnlineProperties);
 
-	USocialManager& GetSocialManager() const;
+	UE_DEPRECATED(5.6, "GetOwningLocalPlayer returns LocalPlayerOwner which is a TWeakObjectPtr and may return nullptr when the local player logs out. Please use the pointer version.")
 	ULocalPlayer& GetOwningLocalPlayer() const;
+
+	ULocalPlayer* GetOwningLocalPlayerPtr() const;
+	USocialManager& GetSocialManager() const;
+
 	const TArray<USocialUser*>& GetAllUsers() const { return AllUsers; }
 
 	/** Finds a SocialUser given a unique net ID from any OSS */
