Change deprecation macros to version 5.5 from 5.6 as UE5-Main is currently 5.5.

#tests Compiled DevelopmentEditor Win64 / Test Client

diff --git a/Engine/Plugins/Online/OnlineFramework/Source/Party/Public/Party/PartyMember.h b/Engine/Plugins/Online/OnlineFramework/Source/Party/Public/Party/PartyMember.h
--- a/Engine/Plugins/Online/OnlineFramework/Source/Party/Public/Party/PartyMember.h
+++ b/Engine/Plugins/Online/OnlineFramework/Source/Party/Public/Party/PartyMember.h
@@ -176,16 +176,16 @@ public:
 
 	virtual void BeginDestroy() override;
 
-	UE_DEPRECATED(5.6, "This has been deprecated to support multiple local players and now requires a performing player.")
+	UE_DEPRECATED(5.5, "This has been deprecated to support multiple local players and now requires a performing player.")
 	bool CanPromoteToLeader() const { return false; }
 	bool CanPromoteToLeader(const ULocalPlayer& PerformingPlayer) const;
-	UE_DEPRECATED(5.6, "This has been deprecated to support multiple local players and now requires a performing player.")
+	UE_DEPRECATED(5.5, "This has been deprecated to support multiple local players and now requires a performing player.")
 	bool PromoteToPartyLeader() { return false; }
 	bool PromoteToPartyLeader(const ULocalPlayer& PerformingPlayer);
-	UE_DEPRECATED(5.6, "This has been deprecated to support multiple local players and now requires a performing player.")
+	UE_DEPRECATED(5.5, "This has been deprecated to support multiple local players and now requires a performing player.")
 	bool CanKickFromParty() const { return false; }
 	bool CanKickFromParty(const ULocalPlayer& PerformingPlayer) const;
-	UE_DEPRECATED(5.6, "This has been deprecated to support multiple local players and now requires a performing player.")
+	UE_DEPRECATED(5.5, "This has been deprecated to support multiple local players and now requires a performing player.")
 	bool KickFromParty() { return false; }
 	bool KickFromParty(const ULocalPlayer& PerformingPlayer);
 
diff --git a/Engine/Plugins/Online/OnlineFramework/Source/Party/Public/Party/SocialParty.h b/Engine/Plugins/Online/OnlineFramework/Source/Party/Public/Party/SocialParty.h
--- a/Engine/Plugins/Online/OnlineFramework/Source/Party/Public/Party/SocialParty.h
+++ b/Engine/Plugins/Online/OnlineFramework/Source/Party/Public/Party/SocialParty.h
@@ -129,7 +129,7 @@ public:
 
 	bool ContainsUser(const USocialUser& User) const;
 
-	UE_DEPRECATED(5.6, "GetOwningLocalPlayer returns the Toolkit's LocalPlayerOwner which is a TWeakObjectPtr and may return nullptr when the local player logs out. Please use the pointer version.")
+	UE_DEPRECATED(5.5, "GetOwningLocalPlayer returns the Toolkit's LocalPlayerOwner which is a TWeakObjectPtr and may return nullptr when the local player logs out. Please use the pointer version.")
 	ULocalPlayer& GetOwningLocalPlayer() const;
 
 	ULocalPlayer* GetOwningLocalPlayerPtr() const;
@@ -239,16 +239,16 @@ protected:
 	// User/member-specific actions that are best exposed on the individuals themselves, but best handled by the actual party
 	bool HasUserBeenInvited(const USocialUser& User) const;
 
-	UE_DEPRECATED(5.6, "This has been deprecated to support multiple local players and now requires a performing player.")
+	UE_DEPRECATED(5.5, "This has been deprecated to support multiple local players and now requires a performing player.")
 	bool CanPromoteMember(const UPartyMember& PartyMember) const { return false; }
 	bool CanPromoteMember(const ULocalPlayer& PerformingPlayer, const UPartyMember& PartyMember) const;
-	UE_DEPRECATED(5.6, "This has been deprecated to support multiple local players and now requires a performing player.")
+	UE_DEPRECATED(5.5, "This has been deprecated to support multiple local players and now requires a performing player.")
 	virtual bool CanKickMember(const UPartyMember& PartyMember) const { return false; }
 	virtual bool CanKickMember(const ULocalPlayer& PerformingPlayer, const UPartyMember& PartyMember) const;
-	UE_DEPRECATED(5.6, "This has been deprecated to support multiple local players and now requires a performing player.")
+	UE_DEPRECATED(5.5, "This has been deprecated to support multiple local players and now requires a performing player.")
 	bool TryPromoteMember(const UPartyMember& PartyMember) { return false; }
 	bool TryPromoteMember(const ULocalPlayer& PerformingPlayer, const UPartyMember& PartyMember);
-	UE_DEPRECATED(5.6, "This has been deprecated to support multiple local players and now requires a performing player.")
+	UE_DEPRECATED(5.5, "This has been deprecated to support multiple local players and now requires a performing player.")
 	virtual bool TryKickMember(const UPartyMember& PartyMember) { return false; }
 	virtual bool TryKickMember(const ULocalPlayer& PerformingPlayer, const UPartyMember& PartyMember);
 	
@@ -274,10 +274,10 @@ protected:
 	virtual ESocialPartyInviteFailureReason CanInviteUserInternal(const USocialUser& User) const;
 	virtual ESocialPartyInviteFailureReason CanInviteUserInternal(const USocialUser& User, const ESocialPartyInviteMethod InviteMethod) const;
 
-	UE_DEPRECATED(5.6, "This has been deprecated to support multiple local players and now requires a performing player.")
+	UE_DEPRECATED(5.5, "This has been deprecated to support multiple local players and now requires a performing player.")
 	virtual bool CanPromoteMemberInternal(const UPartyMember& PartyMember) const { return false; }
 	virtual bool CanPromoteMemberInternal(const ULocalPlayer& PerformingPlayer, const UPartyMember& PartyMember) const;
-	UE_DEPRECATED(5.6, "This has been deprecated to support multiple local players and now requires a performing player.")
+	UE_DEPRECATED(5.5, "This has been deprecated to support multiple local players and now requires a performing player.")
 	virtual bool CanKickMemberInternal(const UPartyMember& PartyMember) const { return false; }
 	virtual bool CanKickMemberInternal(const ULocalPlayer& PerformingPlayer, const UPartyMember& PartyMember) const;
 
diff --git a/Engine/Plugins/Online/OnlineFramework/Source/Party/Public/SocialToolkit.h b/Engine/Plugins/Online/OnlineFramework/Source/Party/Public/SocialToolkit.h
--- a/Engine/Plugins/Online/OnlineFramework/Source/Party/Public/SocialToolkit.h
+++ b/Engine/Plugins/Online/OnlineFramework/Source/Party/Public/SocialToolkit.h
@@ -78,7 +78,7 @@ public:
 	void SetLocalUserOnlineState(EOnlinePresenceState::Type OnlineState);
 	void AddLocalUserOnlineProperties(FPresenceProperties OnlineProperties);
 
-	UE_DEPRECATED(5.6, "GetOwningLocalPlayer returns LocalPlayerOwner which is a TWeakObjectPtr and may return nullptr when the local player logs out. Please use the pointer version.")
+	UE_DEPRECATED(5.5, "GetOwningLocalPlayer returns LocalPlayerOwner which is a TWeakObjectPtr and may return nullptr when the local player logs out. Please use the pointer version.")
 	ULocalPlayer& GetOwningLocalPlayer() const;
 
 	ULocalPlayer* GetOwningLocalPlayerPtr() const;
