SocialChatManager & SocialPartyChatRoom are now easier to subclass.

diff --git a/Engine/Plugins/Online/OnlineFramework/Source/Party/Public/Chat/SocialChatManager.h b/Engine/Plugins/Online/OnlineFramework/Source/Party/Public/Chat/SocialChatManager.h
--- a/Engine/Plugins/Online/OnlineFramework/Source/Party/Public/Chat/SocialChatManager.h
+++ b/Engine/Plugins/Online/OnlineFramework/Source/Party/Public/Chat/SocialChatManager.h
@@ -147,7 +147,7 @@ protected:
 
 	virtual void OnChannelCreatedInternal(USocialChatChannel& CreatedChannel);
 	virtual void OnChannelLeftInternal(USocialChatChannel& ChannelLeft);
-private:
+protected:
 	TMap < ESocialChannelType, TWeakObjectPtr<USocialChatChannel>> ChannelsByType;
 
 	USocialChatRoom& FindOrCreateRoom(const FChatRoomId& RoomId);
@@ -170,7 +170,7 @@ private:
 	virtual void HandleChatRoomJoinPrivateFailure(const FUniqueNetId& LocalUserId, const FChatRoomId& RoomId, const FString& Error) { }
 	virtual void HandleChatRoomExitFailure(const FUniqueNetId& LocalUserId, const FChatRoomId& RoomId, const FString& Error) { }
 
-private:
+protected:
 	UPROPERTY()
 	TMap<TWeakObjectPtr<USocialUser>, TObjectPtr<USocialPrivateMessageChannel>> DirectChannelsByTargetUser;
 
diff --git a/Engine/Plugins/Online/OnlineFramework/Source/Party/Public/Chat/SocialPartyChatRoom.h b/Engine/Plugins/Online/OnlineFramework/Source/Party/Public/Chat/SocialPartyChatRoom.h
--- a/Engine/Plugins/Online/OnlineFramework/Source/Party/Public/Chat/SocialPartyChatRoom.h
+++ b/Engine/Plugins/Online/OnlineFramework/Source/Party/Public/Chat/SocialPartyChatRoom.h
@@ -20,7 +20,7 @@ class PARTY_API USocialPartyChatRoom : public USocialChatRoom
 public:
 	virtual void Initialize(USocialUser* InSocialUser, const FChatRoomId& InChannelId, ESocialChannelType InSourceChannelType);
 
-private:
-	void HandlePartyMemberLeft(EMemberExitedReason Reason);
-	void HandlePartyMemberJoined(UPartyMember& NewMember);
+protected:
+	virtual void HandlePartyMemberLeft(EMemberExitedReason Reason);
+	virtual void HandlePartyMemberJoined(UPartyMember& NewMember);
 };
