Add UE::OnlineFramework::GetLocalPartyMemberToolkits to get all the social toolkits for local players in the party.

diff --git a/Engine/Plugins/Online/OnlineFramework/Source/Party/Private/Party/SocialParty.cpp b/Engine/Plugins/Online/OnlineFramework/Source/Party/Private/Party/SocialParty.cpp
--- a/Engine/Plugins/Online/OnlineFramework/Source/Party/Private/Party/SocialParty.cpp
+++ b/Engine/Plugins/Online/OnlineFramework/Source/Party/Private/Party/SocialParty.cpp
@@ -2099,25 +2099,41 @@ void USocialParty::RunJoinInProgressTimer()
 	}
 }
 
-namespace UE::OnlineFramework::Party
+namespace UE::OnlineFramework
 {
-TArray<FUniqueNetIdRepl> GetPartyMemberIds(const USocialParty* SocialParty)
+TArray<FUniqueNetIdRepl> GetPartyMemberIds(const USocialParty& SocialParty)
 {
 	TArray<FUniqueNetIdRepl> PartyMemberIds;
-	if (SocialParty)
+	auto IsPartyMemberValid = [](const UPartyMember* PartyMember)
 	{
-		auto IsPartyMemberValid = [](const UPartyMember* PartyMember)
-		{
-			CA_ASSUME(PartyMember); // GetPartyMembers filters null members
-			return PartyMember->GetPrimaryNetId().IsValid();
-		};
-		auto GetPartyMemberId = [](const UPartyMember* PartyMember)
+		CA_ASSUME(PartyMember); // GetPartyMembers filters null members
+		return PartyMember->GetPrimaryNetId().IsValid();
+	};
+	auto GetPartyMemberId = [](const UPartyMember* PartyMember)
+	{
+		CA_ASSUME(PartyMember); // GetPartyMembers filters null members
+		return PartyMember->GetPrimaryNetId();
+	};
+	Algo::TransformIf(SocialParty.GetPartyMembers(), PartyMemberIds, IsPartyMemberValid, GetPartyMemberId);
+	return PartyMemberIds;
+}
+
+TArray<USocialToolkit*> GetLocalPartyMemberToolkits(const USocialParty& SocialParty)
+{
+	TArray<USocialToolkit*> SocialToolkits;
+	TArray<UPartyMember*> PartyMembers = SocialParty.GetPartyMembers<UPartyMember>();
+	for (const UPartyMember* PartyMember : PartyMembers)
+	{
+		CA_ASSUME(PartyMember); // GetPartyMembers filters null members
+		if (PartyMember->IsLocalPlayer())
 		{
-			CA_ASSUME(PartyMember); // GetPartyMembers filters null members
-			return PartyMember->GetPrimaryNetId();
-		};
-		Algo::TransformIf(SocialParty->GetPartyMembers(), PartyMemberIds, IsPartyMemberValid, GetPartyMemberId);
+			if (USocialToolkit* SocialToolkit = SocialParty.GetSocialManager().GetSocialToolkit(PartyMember->GetPrimaryNetId()))
+			{
+				SocialToolkits.Add(SocialToolkit);
+			}
+		}
 	}
-	return PartyMemberIds;
+	return SocialToolkits;
 }
+
 } // UE::OnlineFramework::Party
diff --git a/Engine/Plugins/Online/OnlineFramework/Source/Party/Public/Party/SocialParty.h b/Engine/Plugins/Online/OnlineFramework/Source/Party/Public/Party/SocialParty.h
--- a/Engine/Plugins/Online/OnlineFramework/Source/Party/Public/Party/SocialParty.h
+++ b/Engine/Plugins/Online/OnlineFramework/Source/Party/Public/Party/SocialParty.h
@@ -506,9 +506,23 @@ private:
 	mutable FOnInitializationCompletePreNotify OnInitializationCompletePreNotifyEvent;
 };
 
+namespace UE::OnlineFramework
+{
+PARTY_API TArray<FUniqueNetIdRepl> GetPartyMemberIds(const USocialParty& SocialParty);
+PARTY_API TArray<USocialToolkit*> GetLocalPartyMemberToolkits(const USocialParty& SocialParty);
+} // UE::OnlineFramework::Party
+
 namespace UE::OnlineFramework::Party
 {
-PARTY_API TArray<FUniqueNetIdRepl> GetPartyMemberIds(const USocialParty* SocialParty);
+UE_DEPRECATED(5.5, "Use UE::OnlineFramework::GetPartyMemberIds")
+inline TArray<FUniqueNetIdRepl> GetPartyMemberIds(const USocialParty* SocialParty)
+{
+	if (SocialParty)
+	{
+		return UE::OnlineFramework::GetPartyMemberIds(*SocialParty);
+	}
+	return {};
+}
 } // UE::OnlineFramework::Party
 
 #if UE_ENABLE_INCLUDE_ORDER_DEPRECATED_IN_5_2
