Delay social user intialization on party member until a local player logs into the toolkit

diff --git a/Engine/Plugins/Online/OnlineFramework/Source/Party/Private/Party/PartyMember.cpp b/Engine/Plugins/Online/OnlineFramework/Source/Party/Private/Party/PartyMember.cpp
--- a/Engine/Plugins/Online/OnlineFramework/Source/Party/Private/Party/PartyMember.cpp
+++ b/Engine/Plugins/Online/OnlineFramework/Source/Party/Private/Party/PartyMember.cpp
@@ -471,6 +471,25 @@ void UPartyMember::HandleMemberAttributeChanged(const FUniqueNetId& ChangedUserI
 }
 
 void UPartyMember::OnSocialToolkitCreated(USocialToolkit& Toolkit)
+{
+	if (Toolkit.IsOwnerLoggedIn())
+	{
+		OnSocialToolkitLoggedIn(Toolkit);
+	}
+	else
+	{
+		Toolkit.OnLoginChanged().AddWeakLambda(this, [this, Toolkit = TObjectPtr<USocialToolkit>(&Toolkit)](bool bLoggedIn)
+		{
+			if (bLoggedIn)
+			{
+				OnSocialToolkitLoggedIn(*Toolkit);
+				Toolkit->OnLoginChanged().RemoveAll(this);
+			}
+		});
+	}
+}
+
+void UPartyMember::OnSocialToolkitLoggedIn(USocialToolkit& Toolkit)
 {
 	InitializeSocialUserForToolkit(Toolkit);
 }
diff --git a/Engine/Plugins/Online/OnlineFramework/Source/Party/Private/SocialToolkit.cpp b/Engine/Plugins/Online/OnlineFramework/Source/Party/Private/SocialToolkit.cpp
--- a/Engine/Plugins/Online/OnlineFramework/Source/Party/Private/SocialToolkit.cpp
+++ b/Engine/Plugins/Online/OnlineFramework/Source/Party/Private/SocialToolkit.cpp
@@ -661,6 +661,7 @@ void USocialToolkit::OnOwnerLoggedIn()
 	{
 		QueryRecentPlayers();
 	}
+	OnLoginChanged().Broadcast(true);
 }
 
 void USocialToolkit::OnOwnerLoggedOut()
@@ -716,6 +717,7 @@ void USocialToolkit::OnOwnerLoggedOut()
 	LocalUser = NewObject<USocialUser>(this, SocialUserClass);
 
 	OnToolkitReset().Broadcast();
+	OnLoginChanged().Broadcast(false);
 }
 
 void USocialToolkit::QueryFriendsLists()
diff --git a/Engine/Plugins/Online/OnlineFramework/Source/Party/Public/Party/PartyMember.h b/Engine/Plugins/Online/OnlineFramework/Source/Party/Public/Party/PartyMember.h
--- a/Engine/Plugins/Online/OnlineFramework/Source/Party/Public/Party/PartyMember.h
+++ b/Engine/Plugins/Online/OnlineFramework/Source/Party/Public/Party/PartyMember.h
@@ -257,6 +257,7 @@ private:
 	void HandleMemberAttributeChanged(const FUniqueNetId& ChangedUserId, const FString& Attribute, const FString& NewValue, const FString& OldValue);
 	void OnSocialToolkitCreated(USocialToolkit& Toolkit);
 	void OnSocialToolkitDestroyed(USocialToolkit& Toolkit);
+	void OnSocialToolkitLoggedIn(USocialToolkit& Toolkit);
 
 	FOnlinePartyMemberConstPtr OssPartyMember;
 
diff --git a/Engine/Plugins/Online/OnlineFramework/Source/Party/Public/SocialToolkit.h b/Engine/Plugins/Online/OnlineFramework/Source/Party/Public/SocialToolkit.h
--- a/Engine/Plugins/Online/OnlineFramework/Source/Party/Public/SocialToolkit.h
+++ b/Engine/Plugins/Online/OnlineFramework/Source/Party/Public/SocialToolkit.h
@@ -112,6 +112,10 @@ public:
 
 	const FString& GetRecentPlayerNamespaceToQuery() const { return RecentPlayerNamespaceToQuery; }
 	
+	/** Event triggered when the owning local player's login status changes */
+	DECLARE_MULTICAST_DELEGATE_OneParam(FOnLoginChanged, bool /*bLoggedIn*/);
+	FOnLoginChanged& OnLoginChanged() const { return OnLoginChangedEvent; }
+
 	DECLARE_EVENT_OneParam(USocialToolkit, FPartyInviteEvent, USocialUser&);
 	FPartyInviteEvent& OnPartyInviteReceived() const { return OnPartyInviteReceivedEvent; }
 	FPartyInviteEvent& OnPartyInviteRemoved() const { return OnPartyInviteRemovedEvent; }
@@ -321,6 +325,8 @@ private:
 	TSet<IOnlinePartyJoinInfoConstRef> PartyInvitations;
 	mutable TArray<TWeakPtr<FSocialUserList>> CachedSocialUserLists;
 
+	mutable FOnLoginChanged OnLoginChangedEvent;
+
 	mutable FPartyInviteEvent OnPartyInviteReceivedEvent;
 	mutable FPartyInviteEvent OnPartyInviteAcceptedEvent;
 	mutable FPartyInviteEvent OnPartyInviteRejectedEvent;
