- Made new cvar for new dynamic layers for hotfixes, while we track down an issue

diff --git a/Engine/Plugins/Online/OnlineFramework/Source/Hotfix/Private/OnlineHotfixManager.cpp b/Engine/Plugins/Online/OnlineFramework/Source/Hotfix/Private/OnlineHotfixManager.cpp
--- a/Engine/Plugins/Online/OnlineFramework/Source/Hotfix/Private/OnlineHotfixManager.cpp
+++ b/Engine/Plugins/Online/OnlineFramework/Source/Hotfix/Private/OnlineHotfixManager.cpp
@@ -40,6 +40,13 @@ DEFINE_LOG_CATEGORY(LogHotfixManager);
 
 FName NAME_HotfixManager(TEXT("HotfixManager"));
 
+static TAutoConsoleVariable<int32> CVarUseNewDynamicLayersForHotfix(
+	TEXT("ini.UseNewDynamicLayersForHotfix"),
+	0,
+	TEXT("If true, use the new dynamic layers that load/unload configs, specifically for Hotfixes"),
+	ECVF_Default);
+
+
 class FPakFileVisitor : public IPlatformFile::FDirectoryVisitor
 {
 public:
@@ -991,7 +998,7 @@ bool UOnlineHotfixManager::HotfixIniFile(const FString& FileName, const FString&
 	// Flush async loading before modifying GConfig.
 	FlushAsyncLoading();
 	
-	static bool bUseNewDynamicLayers = IConsoleManager::Get().FindConsoleVariable(TEXT("ini.UseNewDynamicLayers"))->GetInt() != 0;
+	static bool bUseNewDynamicLayers = CVarUseNewDynamicLayersForHotfix->GetInt() != 0;
 	if (bUseNewDynamicLayers)
 	{
 		FName Tag = *BuildConfigCacheKey(FileName);
@@ -1232,7 +1239,7 @@ UOnlineHotfixManager::FConfigFileBackup& UOnlineHotfixManager::BackupIniFile(con
 
 void UOnlineHotfixManager::RestoreBackupIniFiles()
 {
-	static bool bUseNewDynamicLayers = IConsoleManager::Get().FindConsoleVariable(TEXT("ini.UseNewDynamicLayers"))->GetInt() != 0;
+	static bool bUseNewDynamicLayers = CVarUseNewDynamicLayersForHotfix->GetInt() != 0;
 	if (bUseNewDynamicLayers)
 	{
 		// @todo branch - would be nice to have a way to know nothing was backed up yet, with Branch mode, so we can skip the FlushAsyncLoading call when there's nothing to do
