- Hopeful fix for new dynamic layers with hotfixes, although still disabled until fully tested

diff --git a/Engine/Plugins/Online/OnlineFramework/Source/Hotfix/Private/OnlineHotfixManager.cpp b/Engine/Plugins/Online/OnlineFramework/Source/Hotfix/Private/OnlineHotfixManager.cpp
--- a/Engine/Plugins/Online/OnlineFramework/Source/Hotfix/Private/OnlineHotfixManager.cpp
+++ b/Engine/Plugins/Online/OnlineFramework/Source/Hotfix/Private/OnlineHotfixManager.cpp
@@ -940,11 +940,11 @@ FConfigBranch* UOnlineHotfixManager::GetBranch(const FString& IniName)
 	const FString StrippedIniNameNoExtension = FPaths::GetBaseFilename(StrippedIniName);
 
 	// find the branch by basename or full filename
-	FConfigBranch* Branch = GConfig->FindBranch(*StrippedIniName, IniName);
+	FConfigBranch* Branch = GConfig->FindBranch(*StrippedIniNameNoExtension, IniName);
 	if (Branch == nullptr)
 	{
 		// does this really happen, seems pretty unexpected
-		UE_LOG(LogHotfixManager, Warning, TEXT("No config braanch found, creating new branch for hotfix %s"), *IniName);
+		UE_LOG(LogHotfixManager, Warning, TEXT("No config branch found, creating new branch for hotfix %s [stripped name is '%s']"), *IniName, *StrippedIniNameNoExtension);
 		
 		Branch = &GConfig->AddNewBranch(IniName);
 	}
