- Actual fix for GetBranch in Hotfix land (cvar still disabled for super duper extra testing first)

diff --git a/Engine/Plugins/Online/OnlineFramework/Source/Hotfix/Private/OnlineHotfixManager.cpp b/Engine/Plugins/Online/OnlineFramework/Source/Hotfix/Private/OnlineHotfixManager.cpp
--- a/Engine/Plugins/Online/OnlineFramework/Source/Hotfix/Private/OnlineHotfixManager.cpp
+++ b/Engine/Plugins/Online/OnlineFramework/Source/Hotfix/Private/OnlineHotfixManager.cpp
@@ -940,7 +940,7 @@ FConfigBranch* UOnlineHotfixManager::GetBranch(const FString& IniName)
 	const FString StrippedIniNameNoExtension = FPaths::GetBaseFilename(StrippedIniName);
 
 	// find the branch by basename or full filename
-	FConfigBranch* Branch = GConfig->FindBranch(*StrippedIniNameNoExtension, IniName);
+	FConfigBranch* Branch = GConfig->FindBranch(*StrippedIniNameNoExtension, StrippedIniNameNoExtension);
 	if (Branch == nullptr)
 	{
 		// does this really happen, seems pretty unexpected
