Fix crash when a SlowTask is created before a SlowTask has started.
Add override to PatchCheck to allow patch checking in editor builds.

diff --git a/Engine/Plugins/Online/OnlineFramework/Source/PatchCheck/Private/PatchCheck.cpp b/Engine/Plugins/Online/OnlineFramework/Source/PatchCheck/Private/PatchCheck.cpp
--- a/Engine/Plugins/Online/OnlineFramework/Source/PatchCheck/Private/PatchCheck.cpp
+++ b/Engine/Plugins/Online/OnlineFramework/Source/PatchCheck/Private/PatchCheck.cpp
@@ -8,6 +8,18 @@
 
 DEFINE_LOG_CATEGORY(LogPatchCheck);
 
+const TCHAR* LexToString(EPatchCheckResult Value)
+{
+	switch (Value)
+	{
+	case EPatchCheckResult::NoPatchRequired:	return TEXT("NoPatchRequired");
+	case EPatchCheckResult::PatchRequired:		return TEXT("PatchRequired");
+	case EPatchCheckResult::NoLoggedInUser:		return TEXT("NoLoggedInUser");
+	default:									checkNoEntry(); // Intentional fallthrough
+	case EPatchCheckResult::PatchCheckFailure:	return TEXT("PatchCheckFailure");
+	}
+}
+
 FPatchCheck& FPatchCheck::Get()
 {
 	static IPatchCheckModule* ConfiguredModule = nullptr;
@@ -244,16 +256,22 @@ bool FPatchCheck::EnvironmentWantsPatchCheck() const
 	return false;
 }
 
+bool FPatchCheck::EditorWantsPatchCheck() const
+{
+	return false;
+}
+
 bool FPatchCheck::SkipPatchCheck() const
 {
 	// Does the environment care about patch checks (LIVE, STAGE, etc)
-	bool bEnvironmentWantsPatchCheck = EnvironmentWantsPatchCheck();
+	const bool bEnvironmentWantsPatchCheck = EnvironmentWantsPatchCheck();
 
 	// Can always opt in to a check
 	const bool bForcePatchCheck = FParse::Param(FCommandLine::Get(), TEXT("ForcePatchCheck"));
 
-	// Prevent a patch check on editor builds 
-	const bool bSkipDueToEditor = UE_EDITOR;
+	// Check whether editor needs a patch check
+	const bool bEditorWantsPatchCheck = UE_EDITOR && EditorWantsPatchCheck();
+	const bool bSkipDueToEditor = !bEditorWantsPatchCheck;
 
 	// Prevent a patch check on dedicated server. UpdateManager also doesn't do a patch check on dedicated server.
 	const bool bSkipDueToDedicatedServer = IsRunningDedicatedServer();
@@ -268,39 +286,43 @@ bool FPatchCheck::SkipPatchCheck() const
 	return bSkipPatchCheck;
 }
 
-void FPatchCheck::OnCheckForPatchComplete(const FUniqueNetId& UniqueId, EUserPrivileges::Type Privilege, uint32 PrivilegeResult, bool bConsoleCheck)
+inline EPatchCheckResult TranslatePatchCheckResult(uint32 PrivilegeResult)
 {
-	UE_LOG(LogPatchCheck, Verbose, TEXT("[OnCheckForPatchComplete] Privilege=%d PrivilegeResult=%d"), (uint32)Privilege, PrivilegeResult);
-
 	EPatchCheckResult Result = EPatchCheckResult::NoPatchRequired;
-	if (Privilege == EUserPrivileges::CanPlayOnline)
+
+	if (PrivilegeResult & (uint32)IOnlineIdentity::EPrivilegeResults::RequiredSystemUpdate)
 	{
-		if (PrivilegeResult & (uint32)IOnlineIdentity::EPrivilegeResults::RequiredSystemUpdate)
-		{
-			Result = EPatchCheckResult::PatchRequired;
-		}
-		else if (PrivilegeResult & (uint32)IOnlineIdentity::EPrivilegeResults::RequiredPatchAvailable)
-		{
-			Result = EPatchCheckResult::PatchRequired;
-		}
-		else if (PrivilegeResult & ((uint32)IOnlineIdentity::EPrivilegeResults::UserNotLoggedIn | (uint32)IOnlineIdentity::EPrivilegeResults::UserNotFound))
+		Result = EPatchCheckResult::PatchRequired;
+	}
+	else if (PrivilegeResult & (uint32)IOnlineIdentity::EPrivilegeResults::RequiredPatchAvailable)
+	{
+		Result = EPatchCheckResult::PatchRequired;
+	}
+	else if (PrivilegeResult & ((uint32)IOnlineIdentity::EPrivilegeResults::UserNotLoggedIn | (uint32)IOnlineIdentity::EPrivilegeResults::UserNotFound))
+	{
+		Result = EPatchCheckResult::NoLoggedInUser;
+	}
+	else if (PrivilegeResult & (uint32)IOnlineIdentity::EPrivilegeResults::GenericFailure)
+	{
+		CA_CONSTANT_IF(PATCH_CHECK_FAIL_ON_GENERIC_FAILURE)
 		{
-			Result = EPatchCheckResult::NoLoggedInUser;
+			Result = EPatchCheckResult::PatchCheckFailure;
 		}
-		else if (PrivilegeResult & (uint32)IOnlineIdentity::EPrivilegeResults::GenericFailure)
+		else
 		{
-			CA_CONSTANT_IF(PATCH_CHECK_FAIL_ON_GENERIC_FAILURE)
-			{
-				Result = EPatchCheckResult::PatchCheckFailure;
-			}
-			else
-			{
-				// Skip console backend failures
-				Result = EPatchCheckResult::NoPatchRequired;
-			}
+			// Skip console backend failures
+			Result = EPatchCheckResult::NoPatchRequired;
 		}
 	}
 
+	return Result;
+}
+
+void FPatchCheck::OnCheckForPatchComplete(const FUniqueNetId& UniqueId, EUserPrivileges::Type Privilege, uint32 PrivilegeResult, bool bConsoleCheck)
+{
+	const EPatchCheckResult Result = Privilege == EUserPrivileges::CanPlayOnline ? TranslatePatchCheckResult(PrivilegeResult) : EPatchCheckResult::NoPatchRequired;
+	UE_LOG(LogPatchCheck, Verbose, TEXT("[OnCheckForPatchComplete] Type: %s, Privilege: %d, PrivilegeResult: %d, PatchCheckResult: %s"), bConsoleCheck ? TEXT("PlatformOSS") : TEXT("DefaultOSS"), (uint32)Privilege, PrivilegeResult, LexToString(Result));
+
 	if (bCheckOSSForUpdate && bConsoleCheck && Result == EPatchCheckResult::NoPatchRequired)
 	{
 		// We perform both checks in this case
diff --git a/Engine/Plugins/Online/OnlineFramework/Source/PatchCheck/Public/PatchCheck.h b/Engine/Plugins/Online/OnlineFramework/Source/PatchCheck/Public/PatchCheck.h
--- a/Engine/Plugins/Online/OnlineFramework/Source/PatchCheck/Public/PatchCheck.h
+++ b/Engine/Plugins/Online/OnlineFramework/Source/PatchCheck/Public/PatchCheck.h
@@ -23,6 +23,7 @@ enum class EPatchCheckResult : uint8
 	PatchCheckFailure,
 	Count,
 };
+PATCHCHECK_API const TCHAR* LexToString(EPatchCheckResult Value);
 
 DECLARE_MULTICAST_DELEGATE_OneParam(FOnPatchCheckComplete, EPatchCheckResult /*Result*/);
 
@@ -58,6 +59,7 @@ protected:
 	virtual void StartOSSPatchCheck();
 	void HandleOSSPatchCheck();
 	virtual bool EnvironmentWantsPatchCheck() const;
+	virtual bool EditorWantsPatchCheck() const;
 	bool SkipPatchCheck() const;
 	void OnCheckForPatchComplete(const FUniqueNetId& UniqueId, EUserPrivileges::Type Privilege, uint32 PrivilegeResult, bool bConsoleCheck);
 	virtual void PatchCheckComplete(EPatchCheckResult PatchResult);
diff --git a/Engine/Source/Runtime/Core/Private/ProfilingDebugging/StallDetector.cpp b/Engine/Source/Runtime/Core/Private/ProfilingDebugging/StallDetector.cpp
--- a/Engine/Source/Runtime/Core/Private/ProfilingDebugging/StallDetector.cpp
+++ b/Engine/Source/Runtime/Core/Private/ProfilingDebugging/StallDetector.cpp
@@ -753,7 +753,7 @@ namespace UE::StallDetector::Private
 	static void HandleSlowTaskFinalize(const FText& TaskName, double DurationInSeconds)
 	{
 		FSlowTaskStallDetectorPause* PauseState = static_cast<FSlowTaskStallDetectorPause*>(FPlatformTLS::GetTlsValue(SlowTaskPauseSlot.GetSlot()));
-		FSlowTaskStallDetectorPause* Parent = PauseState->GetParent();
+		FSlowTaskStallDetectorPause* Parent = ensure(PauseState) ? PauseState->GetParent() : nullptr;
 		FPlatformTLS::SetTlsValue(SlowTaskPauseSlot.GetSlot(), Parent);
 		delete PauseState;
 	}
