Deprecate IOnlinePartySystem::RejoinParty
+ These code paths are not tested by any use case and are subject to code rot
+ Deprecate support from SocialParty/SocialManager. Default USocialParty::bEnableAutomaticPartyRejoin to false
+ A proper implementation should use IOnlinePartySystem::JoinParty
Deprecate overload of IOnlinePartySystem::LeaveParty that does not take a bool synchronize leave as a parameter.
+ Update call sites to pass in true

diff --git a/Engine/Plugins/Online/OnlineFramework/Source/Party/Private/Party/SocialParty.cpp b/Engine/Plugins/Online/OnlineFramework/Source/Party/Private/Party/SocialParty.cpp
--- a/Engine/Plugins/Online/OnlineFramework/Source/Party/Private/Party/SocialParty.cpp
+++ b/Engine/Plugins/Online/OnlineFramework/Source/Party/Private/Party/SocialParty.cpp
@@ -206,10 +206,12 @@ FPartyJoinApproval USocialParty::EvaluateJoinRequest(const TArray<IOnlinePartyUs
 	return JoinApproval;
 }
 
+PRAGMA_DISABLE_DEPRECATION_WARNINGS
 bool USocialParty::ShouldCacheForRejoinOnDisconnect() const
 {
 	return bEnableAutomaticPartyRejoin && GetNumPartyMembers() > 1;
 }
+PRAGMA_ENABLE_DEPRECATION_WARNINGS
 
 bool USocialParty::IsCurrentlyLeaving() const
 {
@@ -1337,7 +1339,7 @@ void USocialParty::LeaveParty(const FOnLeavePartyAttemptComplete& OnLeaveAttempt
 			// All local players will be removed as a consequence of leaving the party with the primary player
 			const IOnlinePartyPtr PartyInterface = Online::GetPartyInterfaceChecked(GetWorld());
 			FOnLeavePartyComplete OnLeaveComplete = FOnLeavePartyComplete::CreateUObject(this, &USocialParty::HandleLeavePartyComplete, OnLeaveAttemptComplete);
-			PartyInterface->LeaveParty(*OwningLocalUserId, PartyId, OnLeaveComplete);
+			PartyInterface->LeaveParty(*OwningLocalUserId, PartyId, true, OnLeaveComplete);
 		}
 		else
 		{
@@ -1362,7 +1364,7 @@ void USocialParty::RemoveLocalMember(const FUniqueNetIdRepl& LocalUserId, const
 		{
 			const IOnlinePartyPtr PartyInterface = Online::GetPartyInterfaceChecked(GetWorld());
 			FOnLeavePartyComplete OnLeaveComplete = FOnLeavePartyComplete::CreateUObject(this, &USocialParty::HandleRemoveLocalPlayerComplete, OnLeaveAttemptComplete);
-			PartyInterface->LeaveParty(*LocalUserId, PartyId, OnLeaveComplete);
+			PartyInterface->LeaveParty(*LocalUserId, PartyId, true, OnLeaveComplete);
 		}
 		else
 		{
diff --git a/Engine/Plugins/Online/OnlineFramework/Source/Party/Private/SocialDebugTools.cpp b/Engine/Plugins/Online/OnlineFramework/Source/Party/Private/SocialDebugTools.cpp
--- a/Engine/Plugins/Online/OnlineFramework/Source/Party/Private/SocialDebugTools.cpp
+++ b/Engine/Plugins/Online/OnlineFramework/Source/Party/Private/SocialDebugTools.cpp
@@ -415,7 +415,7 @@ void USocialDebugTools::LeaveParty(const FString& Instance, const FLeavePartyCom
 				TSharedPtr<const FOnlineParty> ExistingParty = OnlineParty->GetParty(*UserId, IOnlinePartySystem::GetPrimaryPartyTypeId());
 				if (ExistingParty.IsValid())
 				{
-					OnlineParty->LeaveParty(*UserId, *ExistingParty->PartyId, FOnLeavePartyComplete::CreateLambda([this, Instance, OnComplete](const FUniqueNetId& LocalUserId, const FOnlinePartyId& PartyId, const ELeavePartyCompletionResult Result)
+					OnlineParty->LeaveParty(*UserId, *ExistingParty->PartyId, true, FOnLeavePartyComplete::CreateLambda([this, Instance, OnComplete](const FUniqueNetId& LocalUserId, const FOnlinePartyId& PartyId, const ELeavePartyCompletionResult Result)
 					{
 						OnComplete.ExecuteIfBound(Result == ELeavePartyCompletionResult::Succeeded);
 					}));
diff --git a/Engine/Plugins/Online/OnlineFramework/Source/Party/Private/SocialManager.cpp b/Engine/Plugins/Online/OnlineFramework/Source/Party/Private/SocialManager.cpp
--- a/Engine/Plugins/Online/OnlineFramework/Source/Party/Private/SocialManager.cpp
+++ b/Engine/Plugins/Online/OnlineFramework/Source/Party/Private/SocialManager.cpp
@@ -34,6 +34,7 @@ static TAutoConsoleVariable<bool> CVarForceDisconnectedToPartyService(
 // FRejoinableParty
 //////////////////////////////////////////////////////////////////////////
 
+PRAGMA_DISABLE_DEPRECATION_WARNINGS
 USocialManager::FRejoinableParty::FRejoinableParty(const USocialParty& SourceParty)
 	: PartyId(SourceParty.GetPartyId().AsShared()),
 	OriginalJoinMethod(SourceParty.GetOwningLocalMember().GetRepData().GetJoinMethod())
@@ -51,6 +52,7 @@ USocialManager::FRejoinableParty::FRejoinableParty(const USocialParty& SourcePar
 		}		
 	}
 }
+PRAGMA_ENABLE_DEPRECATION_WARNINGS
 
 //////////////////////////////////////////////////////////////////////////
 // FJoinPartyAttempt
@@ -63,14 +65,17 @@ USocialManager::FJoinPartyAttempt::FJoinPartyAttempt(const USocialUser* InTarget
 	, OnJoinComplete(InOnJoinComplete)
 {}
 
+PRAGMA_DISABLE_DEPRECATION_WARNINGS
 USocialManager::FJoinPartyAttempt::FJoinPartyAttempt(TSharedRef<const FRejoinableParty> InRejoinInfo)
 	: PartyTypeId(IOnlinePartySystem::GetPrimaryPartyTypeId())
 	, JoinMethod(InRejoinInfo->OriginalJoinMethod)
 	, RejoinInfo(InRejoinInfo)
 {}
+PRAGMA_ENABLE_DEPRECATION_WARNINGS
 
 FString USocialManager::FJoinPartyAttempt::ToDebugString() const
 {
+PRAGMA_DISABLE_DEPRECATION_WARNINGS
 	return FString::Printf(TEXT("IsRejoin (%s), TargetUser (%s), PartyId (%s), TypeId (%d), TargetUserPlatformId (%s), JoinMethod (%s)"),
 		RejoinInfo.IsValid() ? TEXT("true") : TEXT("false"),
 		TargetUser.IsValid() ? *TargetUser->ToDebugString() : TEXT("invalid"),
@@ -78,6 +83,7 @@ FString USocialManager::FJoinPartyAttempt::ToDebugString() const
 		PartyTypeId.GetValue(),
 		*TargetUserPlatformId.ToDebugString(),
 		*JoinMethod.ToString());
+PRAGMA_ENABLE_DEPRECATION_WARNINGS
 }
 
 const FName USocialManager::FJoinPartyAttempt::Step_FindPlatformSession = TEXT("FindPlatformSession");
@@ -249,7 +255,9 @@ void USocialManager::ShutdownSocialManager()
 	ShutdownPartiesFunc(JoinedPartiesByTypeId);
 	ShutdownPartiesFunc(LeavingPartiesByTypeId);
 
+PRAGMA_DISABLE_DEPRECATION_WARNINGS
 	RejoinableParty.Reset();
+PRAGMA_ENABLE_DEPRECATION_WARNINGS
 
 	if (SocialDebugTools)
 	{
@@ -697,12 +705,13 @@ ECrossplayPreference USocialManager::GetCrossplayPreference() const
 	return ECrossplayPreference::NoSelection;
 }
 
+PRAGMA_DISABLE_DEPRECATION_WARNINGS
 bool USocialManager::ShouldTryRejoiningPersistentParty(const FRejoinableParty& InRejoinableParty) const
 {
-	// If we're alone in our persistent party or we don't have one at the moment, go for it (games will likely have more opinions on the matter)
-	USocialParty* PersistentParty = GetPersistentParty();
- 	return (!PersistentParty || PersistentParty->GetNumPartyMembers() == 1) && !JoinAttemptsByTypeId.Contains(IOnlinePartySystem::GetPrimaryPartyTypeId());
+	// Deprecated
+ 	return false;
 }
+PRAGMA_ENABLE_DEPRECATION_WARNINGS
 
 void USocialManager::RefreshCanCreatePartyObjects()
 {
@@ -822,6 +831,7 @@ void USocialManager::JoinPartyInternal(FJoinPartyAttempt& JoinAttempt)
 	{
 		JoinAttempt.ActionTimeTracker.BeginStep(FJoinPartyAttempt::Step_JoinParty);
 
+PRAGMA_DISABLE_DEPRECATION_WARNINGS
 		if (JoinAttempt.RejoinInfo.IsValid())
 		{
 			UE_LOG(LogParty, Verbose, TEXT("Attempting to rejoin party [%s] now."), *JoinAttempt.RejoinInfo->PartyId->ToDebugString());
@@ -830,6 +840,7 @@ void USocialManager::JoinPartyInternal(FJoinPartyAttempt& JoinAttempt)
 			PartyInterface->RejoinParty(*LocalUserId, *JoinAttempt.RejoinInfo->PartyId, IOnlinePartySystem::GetPrimaryPartyTypeId(), JoinAttempt.RejoinInfo->MemberIds, FOnJoinPartyComplete::CreateUObject(this, &USocialManager::HandleJoinPartyComplete, IOnlinePartySystem::GetPrimaryPartyTypeId()));
 		}
 		else
+PRAGMA_ENABLE_DEPRECATION_WARNINGS
 		{
 			PartyInterface->JoinParty(*LocalUserId, *JoinAttempt.JoinInfo, FOnJoinPartyComplete::CreateUObject(this, &USocialManager::HandleJoinPartyComplete, JoinAttempt.JoinInfo->GetPartyTypeId()));
 		}
@@ -1203,7 +1214,7 @@ void USocialManager::HandleJoinPartyComplete(const FUniqueNetId& LocalUserId, co
 		{
 			UE_LOG(LogParty, Error, TEXT("Auto-bailing on party of type [%d] - cannot finish establishing it without a valid FJoinPartyAttempt."), PartyTypeId.GetValue(), ToString(Result));
 			IOnlinePartyPtr PartyInterface = Online::GetPartyInterfaceChecked(GetWorld());
-			PartyInterface->LeaveParty(LocalUserId, PartyId, FOnLeavePartyComplete::CreateUObject(this, &USocialManager::HandleLeavePartyForMissingJoinAttempt, PartyTypeId));
+			PartyInterface->LeaveParty(LocalUserId, PartyId, true, FOnLeavePartyComplete::CreateUObject(this, &USocialManager::HandleLeavePartyForMissingJoinAttempt, PartyTypeId));
 		}
 		else
 		{
@@ -1230,18 +1241,21 @@ void USocialManager::HandlePersistentPartyStateChanged(EPartyState NewState, EPa
 				PersistentParty->LeaveParty();
 			}
 		}
-
+	
+PRAGMA_DISABLE_DEPRECATION_WARNINGS
 		// If we have other members in our party, then we will try to rejoin this when we come back online
 		if (!RejoinableParty.IsValid() && PersistentParty->ShouldCacheForRejoinOnDisconnect())
 		{
 			UE_LOG(LogParty, Log, TEXT("Caching persistent party [%s] for rejoin"), *PersistentParty->GetPartyId().ToDebugString());
 			RejoinableParty = MakeShared<FRejoinableParty>(*PersistentParty);
 		}
+PRAGMA_ENABLE_DEPRECATION_WARNINGS
 	}
 	else if (NewState == EPartyState::Active)
 	{
 		bIsConnectedToPartyService = true;
 
+PRAGMA_DISABLE_DEPRECATION_WARNINGS
 		if (RejoinableParty.IsValid())
 		{
 			if (ShouldTryRejoiningPersistentParty(*RejoinableParty))
@@ -1258,6 +1272,7 @@ void USocialManager::HandlePersistentPartyStateChanged(EPartyState NewState, EPa
 			// This is the only time we would try to rejoin, and it's saved on the join attempt if initiated
 			RejoinableParty.Reset();
 		}
+PRAGMA_ENABLE_DEPRECATION_WARNINGS
 	}
 }
 
@@ -1323,7 +1338,9 @@ void USocialManager::HandlePartyLeft(EMemberExitedReason Reason, USocialParty* L
 			JoinAttempt->ActionTimeTracker.CompleteStep(FJoinPartyAttempt::Step_LeaveCurrentParty);
 
 			// We're in the process of joining another party of the same type - do we know where we're heading yet?
+PRAGMA_DISABLE_DEPRECATION_WARNINGS
 			if (JoinAttempt->JoinInfo.IsValid() || JoinAttempt->RejoinInfo.IsValid())
+PRAGMA_ENABLE_DEPRECATION_WARNINGS
 			{
 				// Join the new party immediately and early out
 				JoinPartyInternal(*JoinAttempt);
diff --git a/Engine/Plugins/Online/OnlineFramework/Source/Party/Public/Party/SocialParty.h b/Engine/Plugins/Online/OnlineFramework/Source/Party/Public/Party/SocialParty.h
--- a/Engine/Plugins/Online/OnlineFramework/Source/Party/Public/Party/SocialParty.h
+++ b/Engine/Plugins/Online/OnlineFramework/Source/Party/Public/Party/SocialParty.h
@@ -228,6 +228,7 @@ protected:
 	bool IsInitialized() const;
 	void TryFinishInitialization();
 
+	UE_DEPRECATED(5.5, "Rejoinable Parties is deprecated and will be removed")
 	bool ShouldCacheForRejoinOnDisconnect() const;
 
 	void SetIsMissingPlatformSession(bool bInIsMissingPlatformSession);
@@ -398,8 +399,9 @@ private:
 	UPROPERTY()
 	TMap<FUniqueNetIdRepl, TObjectPtr<UPartyMember>> PartyMembersById;
 
+	UE_DEPRECATED(5.5, "Rejoinable Parties is deprecated and will be removed")
 	UPROPERTY(config)
-	bool bEnableAutomaticPartyRejoin = true;
+	bool bEnableAutomaticPartyRejoin = false;
 
 	TMap<FUniqueNetIdRepl, double> LastInviteSentById;
 
diff --git a/Engine/Plugins/Online/OnlineFramework/Source/Party/Public/SocialManager.h b/Engine/Plugins/Online/OnlineFramework/Source/Party/Public/SocialManager.h
--- a/Engine/Plugins/Online/OnlineFramework/Source/Party/Public/SocialManager.h
+++ b/Engine/Plugins/Online/OnlineFramework/Source/Party/Public/SocialManager.h
@@ -125,7 +125,7 @@ protected:
 	void JoinParty(const USocialUser& UserToJoin, const FOnlinePartyTypeId& PartyTypeId, const FOnJoinPartyAttemptComplete& OnJoinPartyComplete, const FName& JoinMethod);
 
 protected:
-	struct PARTY_API FRejoinableParty : public TSharedFromThis<FRejoinableParty>
+	struct PARTY_API UE_DEPRECATED(5.5, "Rejoinable Parties is deprecated and will be removed") FRejoinableParty : public TSharedFromThis<FRejoinableParty>
 	{
 		FRejoinableParty(const USocialParty& SourceParty);
 
@@ -137,7 +137,12 @@ protected:
 public:
 	struct PARTY_API FJoinPartyAttempt
 	{
+		UE_DEPRECATED(5.5, "Rejoinable Parties is deprecated and will be removed")
+PRAGMA_DISABLE_DEPRECATION_WARNINGS
 		FJoinPartyAttempt(TSharedRef<const FRejoinableParty> InRejoinInfo);
+		FJoinPartyAttempt(const FJoinPartyAttempt&) = default;
+		FJoinPartyAttempt(FJoinPartyAttempt&&) = default;
+PRAGMA_ENABLE_DEPRECATION_WARNINGS
 		FJoinPartyAttempt(const USocialUser* InTargetUser, const FOnlinePartyTypeId& InPartyTypeId, const FName& InJoinMethod, const FOnJoinPartyAttemptComplete& InOnJoinComplete);
 
 		FString ToDebugString() const;
@@ -147,7 +152,10 @@ public:
 		FName JoinMethod = PartyJoinMethod::Unspecified;
 		FUniqueNetIdRepl TargetUserPlatformId;
 
+		UE_DEPRECATED(5.5, "Rejoinable Parties is deprecated and will be removed")
+PRAGMA_DISABLE_DEPRECATION_WARNINGS
 		TSharedPtr<const FRejoinableParty> RejoinInfo;
+PRAGMA_ENABLE_DEPRECATION_WARNINGS
 		IOnlinePartyJoinInfoConstPtr JoinInfo;
 
 		FOnJoinPartyAttemptComplete OnJoinComplete;
@@ -190,7 +198,10 @@ protected:
 	/** Up to the game to decide whether it wants to allow crossplay (generally based on a user setting of some kind) */
 	virtual ECrossplayPreference GetCrossplayPreference() const;
 
+	UE_DEPRECATED(5.5, "Rejoinable Parties is deprecated and will be removed")
+PRAGMA_DISABLE_DEPRECATION_WARNINGS
 	virtual bool ShouldTryRejoiningPersistentParty(const FRejoinableParty& InRejoinableParty) const;
+PRAGMA_ENABLE_DEPRECATION_WARNINGS
 
 	template <typename InteractionT>
 	void RegisterInteraction()
@@ -212,7 +223,10 @@ protected:
 	virtual TSubclassOf<USocialDebugTools> GetSocialDebugToolsClass() const;
 
 	/** Info on the persistent party we were in when losing connection to the party service and want to rejoin when it returns */
+	UE_DEPRECATED(5.5, "Rejoinable Parties is deprecated and will be removed")
+PRAGMA_DISABLE_DEPRECATION_WARNINGS
 	TSharedPtr<FRejoinableParty> RejoinableParty;
+PRAGMA_ENABLE_DEPRECATION_WARNINGS
 
 	/** The desired type of SocialToolkit to create for each local player */
 	TSubclassOf<USocialToolkit> ToolkitClass;
diff --git a/Engine/Plugins/Online/OnlineSubsystem/Source/Private/OnlinePartyInterface.cpp b/Engine/Plugins/Online/OnlineSubsystem/Source/Private/OnlinePartyInterface.cpp
--- a/Engine/Plugins/Online/OnlineSubsystem/Source/Private/OnlinePartyInterface.cpp
+++ b/Engine/Plugins/Online/OnlineSubsystem/Source/Private/OnlinePartyInterface.cpp
@@ -453,7 +453,9 @@ const TCHAR* ToString(const EJoinPartyCompletionResult Value)
 	{
 		return TEXT("LoggedOut");
 	}
+PRAGMA_DISABLE_DEPRECATION_WARNINGS
 	case EJoinPartyCompletionResult::UnableToRejoin:
+PRAGMA_ENABLE_DEPRECATION_WARNINGS
 	{
 		return TEXT("UnableToRejoin");
 	}
@@ -905,3 +907,13 @@ FString IOnlinePartyJoinInfo::ToDebugString() const
 		IsAcceptingMembers() ? 1 : 0,
 		GetNotAcceptingReason());
 };
+
+bool IOnlinePartySystem::RejoinParty(const FUniqueNetId& LocalUserId, const FOnlinePartyId& PartyId, const FOnlinePartyTypeId& PartyTypeId, const TArray<FUniqueNetIdRef>& FormerMembers, const FOnJoinPartyComplete& Delegate)
+{
+	return false;
+}
+
+bool IOnlinePartySystem::LeaveParty(const FUniqueNetId& LocalUserId, const FOnlinePartyId& PartyId, const FOnLeavePartyComplete& Delegate)
+{
+	return LeaveParty(LocalUserId, PartyId, true, Delegate);
+}
diff --git a/Engine/Plugins/Online/OnlineSubsystem/Source/Public/Interfaces/OnlinePartyInterface.h b/Engine/Plugins/Online/OnlineSubsystem/Source/Public/Interfaces/OnlinePartyInterface.h
--- a/Engine/Plugins/Online/OnlineSubsystem/Source/Public/Interfaces/OnlinePartyInterface.h
+++ b/Engine/Plugins/Online/OnlineSubsystem/Source/Public/Interfaces/OnlinePartyInterface.h
@@ -621,7 +621,7 @@ enum class EPartyState : uint8
 	None,
 	CreatePending,
 	JoinPending,
-	RejoinPending, 
+	RejoinPending UE_DEPRECATED(5.5, "RejoinParty is deprecated. Implement this feature via JoinParty."), 
 	LeavePending,
 	Active,
 	Disconnected,
@@ -1158,13 +1158,13 @@ PARTY_DECLARE_DELEGATETYPE(OnPartyRequestToJoinRemoved);
  * Interface definition for the online party services 
  * Allows for forming a party and communicating with party members
  */
-class ONLINESUBSYSTEM_API IOnlinePartySystem
+class IOnlinePartySystem
 {
 protected:
-	IOnlinePartySystem() {};
+	IOnlinePartySystem() = default;
 
 public:
-	virtual ~IOnlinePartySystem() {};
+	virtual ~IOnlinePartySystem() = default;
 
 	/**
 	 * Restore party memberships. Intended to be called once during login to restore state from other running instances.
@@ -1256,7 +1256,7 @@ public:
 	virtual void QueryPartyJoinability(const FUniqueNetId& LocalUserId, const IOnlinePartyJoinInfo& OnlinePartyJoinInfo, const FOnQueryPartyJoinabilityCompleteEx& Delegate) = 0;
 
 	/**
-	 * Attempt to rejoin a former party
+	 * DEPRECATED! Attempt to rejoin a former party. Deprecated: Implement this feature with JoinParty instead.
 	 *
 	 * @param LocalUserId - user making the request
 	 * @param PartyId - id of the party you want to rejoin
@@ -1266,7 +1266,8 @@ public:
 	 *
 	 * @return true if task was started
 	 */
-	virtual bool RejoinParty(const FUniqueNetId& LocalUserId, const FOnlinePartyId& PartyId, const FOnlinePartyTypeId& PartyTypeId, const TArray<FUniqueNetIdRef>& FormerMembers, const FOnJoinPartyComplete& Delegate = FOnJoinPartyComplete()) = 0;
+	UE_DEPRECATED(5.5, "RejoinParty is deprecated. Implement this feature via JoinParty.")
+	virtual bool ONLINESUBSYSTEM_API RejoinParty(const FUniqueNetId& LocalUserId, const FOnlinePartyId& PartyId, const FOnlinePartyTypeId& PartyTypeId, const TArray<FUniqueNetIdRef>& FormerMembers, const FOnJoinPartyComplete& Delegate = FOnJoinPartyComplete());
 
 	/**
 	 * Leave an existing party
@@ -1278,7 +1279,8 @@ public:
 	 *
 	 * @return true if task was started
 	 */
-	virtual bool LeaveParty(const FUniqueNetId& LocalUserId, const FOnlinePartyId& PartyId, const FOnLeavePartyComplete& Delegate = FOnLeavePartyComplete()) = 0;
+	UE_DEPRECATED(5.5, "Use LeaveParty that takes bSynchronizeLeave")
+	virtual bool ONLINESUBSYSTEM_API LeaveParty(const FUniqueNetId& LocalUserId, const FOnlinePartyId& PartyId, const FOnLeavePartyComplete& Delegate = FOnLeavePartyComplete());
 
 	/**
 	 * Leave an existing party
@@ -1885,7 +1887,7 @@ enum class EJoinPartyCompletionResult : int8
 	/** You were logged out while attempting to join the party */
 	LoggedOut,
 	/** You were unable to rejoin the party */
-	UnableToRejoin,
+	UnableToRejoin UE_DEPRECATED(5.5, "RejoinParty is deprecated. Implement this feature via JoinParty."),
 	/** Your platform is not compatible with the party */
 	IncompatiblePlatform,
 
