Restore how default social user is populated

diff --git a/Engine/Plugins/Online/OnlineFramework/Source/Party/Private/Party/PartyMember.cpp b/Engine/Plugins/Online/OnlineFramework/Source/Party/Private/Party/PartyMember.cpp
--- a/Engine/Plugins/Online/OnlineFramework/Source/Party/Private/Party/PartyMember.cpp
+++ b/Engine/Plugins/Online/OnlineFramework/Source/Party/Private/Party/PartyMember.cpp
@@ -124,21 +124,6 @@ bool AreAllSocialUsersInitialized(const UPartyMember& InPartyMember)
 	}
 	return true;
 }
-
-// Get the default social user. Used when it really doesn't matter which social user we get.
-USocialUser& GetDefaultSocialUser(const UPartyMember& InPartyMember)
-{
-	USocialToolkit* OwnerToolkit = InPartyMember.GetParty().GetSocialManager().GetSocialToolkit(InPartyMember.GetPrimaryNetId());
-	// If we are not a local user then we simply get the first local user's toolkit
-	if (OwnerToolkit == nullptr)
-	{
-		OwnerToolkit = InPartyMember.GetParty().GetSocialManager().GetFirstLocalUserToolkit();
-	}
-	check(OwnerToolkit);
-	USocialUser* SocialUser = OwnerToolkit->FindUser(InPartyMember.GetPrimaryNetId());
-	check(SocialUser);
-	return *SocialUser;
-}
 }
 
 UPartyMember::UPartyMember()
@@ -172,7 +157,21 @@ void UPartyMember::InitializePartyMember(const FOnlinePartyMemberConstRef& InOss
 			DebugInitializer = MakeUnique<FDebugInitializer>(*this);
 		}
 
-		DefaultSocialUser = &GetDefaultSocialUser(*this);
+		{
+			USocialToolkit* OwnerToolkit = GetParty().GetSocialManager().GetSocialToolkit(OssPartyMember->GetUserId());
+			// If we are not a local user then we simply get the first local user's toolkit
+			if (OwnerToolkit == nullptr)
+			{
+				OwnerToolkit = GetParty().GetSocialManager().GetFirstLocalUserToolkit();
+			}
+			check(OwnerToolkit);
+
+			OwnerToolkit->QueueUserDependentAction(InOssMember->GetUserId(),
+				[this] (USocialUser& User)
+				{
+					DefaultSocialUser = &User;
+				}, false);
+		}
 
 		// Local player already has all the data they need, everyone else we want to wait for
 		if (IsLocalPlayer())
