Revert changes to USocialParty::GetOwningLocalPlayerPtr to use the default social user
This is being called on logout and the call to GetSocialUser(player id) is returning nullptr.  The quickest fix is to revert this.

diff --git a/Engine/Plugins/Online/OnlineFramework/Source/Party/Private/Party/SocialParty.cpp b/Engine/Plugins/Online/OnlineFramework/Source/Party/Private/Party/SocialParty.cpp
--- a/Engine/Plugins/Online/OnlineFramework/Source/Party/Private/Party/SocialParty.cpp
+++ b/Engine/Plugins/Online/OnlineFramework/Source/Party/Private/Party/SocialParty.cpp
@@ -1404,7 +1404,7 @@ ULocalPlayer& USocialParty::GetOwningLocalPlayer() const
 ULocalPlayer* USocialParty::GetOwningLocalPlayerPtr() const
 {
 	//@todo DanH Party: This is a wee bit heavy - should be able to do this in fewer steps
-	return GetOwningLocalMember().GetSocialUser(GetOwningLocalUserId())->GetOwningToolkit().GetOwningLocalPlayerPtr();
+	return GetOwningLocalMember().GetSocialUser().GetOwningToolkit().GetOwningLocalPlayerPtr();
 }
 
 bool USocialParty::IsLocalPlayerPartyLeader() const
