fix crash when failing to get a valid QOS region

diff --git a/Engine/Plugins/Online/OnlineFramework/Source/Qos/Private/QosRegionManager.cpp b/Engine/Plugins/Online/OnlineFramework/Source/Qos/Private/QosRegionManager.cpp
--- a/Engine/Plugins/Online/OnlineFramework/Source/Qos/Private/QosRegionManager.cpp
+++ b/Engine/Plugins/Online/OnlineFramework/Source/Qos/Private/QosRegionManager.cpp
@@ -511,7 +511,7 @@ void UQosRegionManager::OnQosEvaluationComplete(EQosCompletionResult Result, con
 	
 	if (OutSelectedSubRegion)
 	{
-		*OutSelectedSubRegion = *BestRegionSubregions[0];
+		*OutSelectedSubRegion = BestRegionSubregions.Num() ? *BestRegionSubregions[0] : FString();
 	}
 
 #if DEBUG_SUBCOMPARE_BY_SUBSPACE
