Fix for patch check not firing on client builds.

diff --git a/Engine/Plugins/Online/OnlineFramework/Source/PatchCheck/Private/PatchCheck.cpp b/Engine/Plugins/Online/OnlineFramework/Source/PatchCheck/Private/PatchCheck.cpp
--- a/Engine/Plugins/Online/OnlineFramework/Source/PatchCheck/Private/PatchCheck.cpp
+++ b/Engine/Plugins/Online/OnlineFramework/Source/PatchCheck/Private/PatchCheck.cpp
@@ -270,8 +270,8 @@ bool FPatchCheck::SkipPatchCheck() const
 	const bool bForcePatchCheck = FParse::Param(FCommandLine::Get(), TEXT("ForcePatchCheck"));
 
 	// Check whether editor needs a patch check
-	const bool bEditorWantsPatchCheck = UE_EDITOR && EditorWantsPatchCheck();
-	const bool bSkipDueToEditor = !bEditorWantsPatchCheck;
+	const bool bEditorWantsPatchCheck = EditorWantsPatchCheck();
+	const bool bSkipDueToEditor = UE_EDITOR && !bEditorWantsPatchCheck;
 
 	// Prevent a patch check on dedicated server. UpdateManager also doesn't do a patch check on dedicated server.
 	const bool bSkipDueToDedicatedServer = IsRunningDedicatedServer();
