Change SocialChatManager in SocialToolkit to be protected instead of private

#tests successfully compiled a game and used text chat with this change.

#rnx

diff --git a/Engine/Plugins/Online/OnlineFramework/Source/Party/Public/SocialToolkit.h b/Engine/Plugins/Online/OnlineFramework/Source/Party/Public/SocialToolkit.h
--- a/Engine/Plugins/Online/OnlineFramework/Source/Party/Public/SocialToolkit.h
+++ b/Engine/Plugins/Online/OnlineFramework/Source/Party/Public/SocialToolkit.h
@@ -302,6 +302,10 @@ private:	// Handlers
 	FTSTicker::FDelegateHandle Debug_PresenceTickerHandle;
 #endif
 
+protected:
+	UPROPERTY()
+	TObjectPtr<USocialChatManager> SocialChatManager;
+
 private:
 	static USocialToolkit* GetToolkitForPlayerInternal(const ULocalPlayer* LocalPlayer);
 	static TMap<TWeakObjectPtr<const ULocalPlayer>, TWeakObjectPtr<USocialToolkit>> AllToolkitsByOwningPlayer;
@@ -319,9 +323,6 @@ private:
 	UPROPERTY()
 	TWeakObjectPtr<ULocalPlayer> LocalPlayerOwner = nullptr;
 
-	UPROPERTY()
-	TObjectPtr<USocialChatManager> SocialChatManager;
-
 	TSet<IOnlinePartyJoinInfoConstRef> PartyInvitations;
 	mutable TArray<TWeakPtr<FSocialUserList>> CachedSocialUserLists;
 
