Change hotfix asset patch failure from an error to a warning when the asset is not found.

diff --git a/Engine/Plugins/Online/OnlineFramework/Source/Hotfix/Private/OnlineHotfixManager.cpp b/Engine/Plugins/Online/OnlineFramework/Source/Hotfix/Private/OnlineHotfixManager.cpp
--- a/Engine/Plugins/Online/OnlineFramework/Source/Hotfix/Private/OnlineHotfixManager.cpp
+++ b/Engine/Plugins/Online/OnlineFramework/Source/Hotfix/Private/OnlineHotfixManager.cpp
@@ -1552,7 +1552,10 @@ void UOnlineHotfixManager::PatchAssetsFromIniFiles()
 	}
 	else
 	{
-		UE_LOG(LogHotfixManager, Error, TEXT("Only %i of %i assets were successfully patched from 'AssetHotfix' section in the Game .ini file. The patched assets will be forced to remain loaded. Any assets that failed to patch may be left in an invalid state!"), AssetsHotfixedFromIniFiles.Num(), TotalPatchableAssets);
+		// Some assets will fail to patch due to their plugins not being installed or mounted when PatchAssetsFromIniFiles is called.
+		// This behavior is a limitation of having a hotfix for the overall game and will require plugin specific hotfixes to address.
+		// PatchAssetsFromIniFiles must be called again when plugins change to allow patching assets when their plugin is mounted.
+		UE_LOG(LogHotfixManager, Warning, TEXT("Only %i of %i assets were successfully patched from 'AssetHotfix' section in the Game .ini file. The patched assets will be forced to remain loaded. Any assets that failed to patch may be left in an invalid state! Falure may occur due to the plugin not being installed and mounted."), AssetsHotfixedFromIniFiles.Num(), TotalPatchableAssets);
 	}
 }
 
