Make party invite metadata accessible to analytic events

diff --git a/Engine/Plugins/Online/OnlineFramework/Source/Party/Private/Party/SocialParty.cpp b/Engine/Plugins/Online/OnlineFramework/Source/Party/Private/Party/SocialParty.cpp
--- a/Engine/Plugins/Online/OnlineFramework/Source/Party/Private/Party/SocialParty.cpp
+++ b/Engine/Plugins/Online/OnlineFramework/Source/Party/Private/Party/SocialParty.cpp
@@ -303,6 +303,11 @@ ESocialPartyInviteFailureReason USocialParty::CanInviteUserInternal(const USocia
 	return ESocialPartyInviteFailureReason::Success;
 }
 
+FString USocialParty::GetInviteMetadata(const FString& ExistingMetadata)
+{
+	return ExistingMetadata;
+}
+
 bool USocialParty::IsInviteRateLimited(const USocialUser& User, ESocialSubsystem SubsystemType) const
 {
 	const FUniqueNetIdRepl UserId = User.GetUserId(SubsystemType);
@@ -320,6 +325,7 @@ bool USocialParty::IsInviteRateLimited(const USocialUser& User, ESocialSubsystem
 
 bool USocialParty::TryInviteUser(const USocialUser& UserToInvite, const ESocialPartyInviteMethod InviteMethod, const FString& MetaData)
 {
+	const FString InviteMetadata = GetInviteMetadata(MetaData);	
 	bool bSentInvite = false;
 	ESocialPartyInviteFailureReason CanInviteResult = CanInviteUserInternal(UserToInvite, InviteMethod);
 	if (CanInviteResult == ESocialPartyInviteFailureReason::Success)
@@ -369,23 +375,23 @@ bool USocialParty::TryInviteUser(const USocialUser& UserToInvite, const ESocialP
 				}
 			}
 			ESocialPartyInviteFailureReason FailureReason = bSentPlatformInvite ? ESocialPartyInviteFailureReason::Success : ESocialPartyInviteFailureReason::PlatformInviteFailed;
-			OnInviteSentInternal(ESocialSubsystem::Platform, UserToInvite, bSentPlatformInvite, FailureReason, InviteMethod);
+			OnInviteSentInternal(ESocialSubsystem::Platform, UserToInvite, bSentPlatformInvite, FailureReason, InviteMethod, InviteMetadata);
 			bSentInvite |= bSentPlatformInvite;
 		}
 		if ((!bSentInvite || bMustSendPrimaryInvite) && PerformingPlayer && UserPrimaryId.IsValid() && !IsInviteRateLimited(UserToInvite, ESocialSubsystem::Primary))
 		{
 			// Primary subsystem invites can be sent directly to the user via the party interface
-			const FPartyInvitationRecipient Recipient(*UserPrimaryId, MetaData);
+			const FPartyInvitationRecipient Recipient(*UserPrimaryId, InviteMetadata);
 			const IOnlinePartyPtr PartyInterface = Online::GetPartyInterfaceChecked(GetWorld());
 			const bool bSentPrimaryInvite = PartyInterface->SendInvitation(*PerformingPlayer->GetPreferredUniqueNetId(), GetPartyId(), Recipient);
 			ESocialPartyInviteFailureReason FailureReason = bSentPrimaryInvite ? ESocialPartyInviteFailureReason::Success : ESocialPartyInviteFailureReason::PartyInviteFailed;
-			OnInviteSentInternal(ESocialSubsystem::Primary, UserToInvite, bSentPrimaryInvite, FailureReason, InviteMethod);
+			OnInviteSentInternal(ESocialSubsystem::Primary, UserToInvite, bSentPrimaryInvite, FailureReason, InviteMethod, InviteMetadata);
 			bSentInvite |= bSentPrimaryInvite;
 		}
 	}
 	else
 	{
-		OnInviteSentInternal(ESocialSubsystem::MAX, UserToInvite, false, CanInviteResult, InviteMethod);
+		OnInviteSentInternal(ESocialSubsystem::MAX, UserToInvite, false, CanInviteResult, InviteMethod, InviteMetadata);
 	}
 	return bSentInvite;
 }
@@ -664,7 +670,7 @@ void USocialParty::OnLeftPartyInternal(EMemberExitedReason Reason)
 	OnPartyLeft().Broadcast(Reason);
 }
 
-void USocialParty::OnInviteSentInternal(ESocialSubsystem SubsystemType, const USocialUser& InvitedUser, bool bWasSuccessful, const ESocialPartyInviteFailureReason FailureReason, const ESocialPartyInviteMethod InviteMethod)
+void USocialParty::OnInviteSentInternal(ESocialSubsystem SubsystemType, const USocialUser& InvitedUser, bool bWasSuccessful, const ESocialPartyInviteFailureReason FailureReason, const ESocialPartyInviteMethod InviteMethod, const FString& MetaData)
 {
 	// If the invite is successful, save the current timestamp to stop invites to this user 
 	// for a while (defined in PlatformUserInviteCooldown/PrimaryUserInviteCooldown)
@@ -680,8 +686,11 @@ void USocialParty::OnInviteSentInternal(ESocialSubsystem SubsystemType, const US
 	
 	OnInviteSent().Broadcast(InvitedUser);
 
-	// Call the deprecated method after the current OnInviteSentInternal
+	// Call the deprecated methods after the current OnInviteSentInternal
+PRAGMA_DISABLE_DEPRECATION_WARNINGS
+	OnInviteSentInternal(SubsystemType, InvitedUser, bWasSuccessful, FailureReason, InviteMethod);
 	OnInviteSentInternal(SubsystemType, InvitedUser, bWasSuccessful);
+PRAGMA_ENABLE_DEPRECATION_WARNINGS
 }
 
 UPartyMember* USocialParty::GetOrCreatePartyMember(const FUniqueNetId& MemberId)
diff --git a/Engine/Plugins/Online/OnlineFramework/Source/Party/Private/User/SocialUser.cpp b/Engine/Plugins/Online/OnlineFramework/Source/Party/Private/User/SocialUser.cpp
--- a/Engine/Plugins/Online/OnlineFramework/Source/Party/Private/User/SocialUser.cpp
+++ b/Engine/Plugins/Online/OnlineFramework/Source/Party/Private/User/SocialUser.cpp
@@ -1003,16 +1003,16 @@ bool USocialUser::CanSendFriendInvite(ESocialSubsystem SubsystemType) const
 
 void USocialUser::JoinParty(const FOnlinePartyTypeId& PartyTypeId, const FName& JoinMethod) const
 {
-	const bool bHasSentInvite = HasSentPartyInvite(PartyTypeId);
+	IOnlinePartyJoinInfoConstPtr SentInvite = GetSentPartyInvite(PartyTypeId);
 
 	GetOwningToolkit().GetSocialManager().JoinParty(*this, PartyTypeId, USocialManager::FOnJoinPartyAttemptComplete(), JoinMethod);
 
 	// Regardless of the outcome, note that the invite was accepted (deletes it from the OSS party system)
-	if (bHasSentInvite)
+	if (SentInvite)
 	{
 		IOnlinePartyPtr PartyInterface = Online::GetPartyInterfaceChecked(GetWorld());
 		PartyInterface->ClearInvitations(*GetOwningToolkit().GetLocalUserNetId(ESocialSubsystem::Primary), *GetUserId(ESocialSubsystem::Primary), nullptr);
-		OnPartyInviteAcceptedInternal(PartyTypeId);
+		OnPartyInviteAcceptedInternal(PartyTypeId, SentInvite->GetPlatformData());
 	}
 }
 
@@ -1273,6 +1273,18 @@ bool USocialUser::HasSentPartyInvite(const FOnlinePartyTypeId& PartyTypeId) cons
 	return false;
 }
 
+IOnlinePartyJoinInfoConstPtr USocialUser::GetSentPartyInvite(const FOnlinePartyTypeId& PartyTypeId) const
+{
+	for (const IOnlinePartyJoinInfoConstRef& Invite : ReceivedPartyInvites)
+	{
+		if (Invite->GetPartyTypeId() == PartyTypeId)
+		{
+			return Invite;
+		}
+	}
+	return nullptr;
+}
+
 FUniqueNetIdRepl USocialUser::GetUserId(ESocialSubsystem SubsystemType) const
 {
 	if (const FSubsystemUserInfo* SubsystemInfo = SubsystemInfoByType.Find(SubsystemType))
@@ -1414,8 +1426,13 @@ void USocialUser::OnPresenceChangedInternal(ESocialSubsystem SubsystemType)
 	OnUserPresenceChanged().Broadcast(SubsystemType);
 }
 
-void USocialUser::OnPartyInviteAcceptedInternal(const FOnlinePartyTypeId& PartyTypeId) const
+void USocialUser::OnPartyInviteAcceptedInternal(const FOnlinePartyTypeId& PartyTypeId, const FString& InviteMetadata) const
 {
+	// Call deprecated method
+PRAGMA_DISABLE_DEPRECATION_WARNINGS
+	OnPartyInviteAcceptedInternal(PartyTypeId);
+PRAGMA_ENABLE_DEPRECATION_WARNINGS
+
 	OnPartyInviteAccepted().Broadcast();
 }
 
diff --git a/Engine/Plugins/Online/OnlineFramework/Source/Party/Public/Party/SocialParty.h b/Engine/Plugins/Online/OnlineFramework/Source/Party/Public/Party/SocialParty.h
--- a/Engine/Plugins/Online/OnlineFramework/Source/Party/Public/Party/SocialParty.h
+++ b/Engine/Plugins/Online/OnlineFramework/Source/Party/Public/Party/SocialParty.h
@@ -280,9 +280,12 @@ protected:
 	virtual bool CanKickMemberInternal(const UPartyMember& PartyMember) const { return false; }
 	virtual bool CanKickMemberInternal(const ULocalPlayer& PerformingPlayer, const UPartyMember& PartyMember) const;
 
-	virtual void OnInviteSentInternal(ESocialSubsystem SubsystemType, const USocialUser& InvitedUser, bool bWasSuccessful, const ESocialPartyInviteFailureReason FailureReason, const ESocialPartyInviteMethod InviteMethod);
+	virtual void OnInviteSentInternal(ESocialSubsystem SubsystemType, const USocialUser& InvitedUser, bool bWasSuccessful, const ESocialPartyInviteFailureReason FailureReason, const ESocialPartyInviteMethod InviteMethod, const FString& MetaData);
 
-	/* Deprecated version */
+	UE_DEPRECATED(5.6, "Override 'void OnInviteSentInternal(ESocialSubsystem SubsystemType, const USocialUser& InvitedUser, bool bWasSuccessful, const ESocialPartyInviteFailureReason FailureReason, const ESocialPartyInviteMethod InviteMethod, const FString& MetaData)' instead")
+	virtual void OnInviteSentInternal(ESocialSubsystem SubsystemType, const USocialUser& InvitedUser, bool bWasSuccessful, const ESocialPartyInviteFailureReason FailureReason, const ESocialPartyInviteMethod InviteMethod) {};
+	
+	UE_DEPRECATED(5.6, "Override 'void OnInviteSentInternal(ESocialSubsystem SubsystemType, const USocialUser& InvitedUser, bool bWasSuccessful, const ESocialPartyInviteFailureReason FailureReason, const ESocialPartyInviteMethod InviteMethod, const FString& MetaData)' instead")
 	virtual void OnInviteSentInternal(ESocialSubsystem SubsystemType, const USocialUser& InvitedUser, bool bWasSuccessful) {};
 
 	virtual void HandlePartySystemStateChange(EPartySystemState NewState);
@@ -302,6 +305,9 @@ protected:
 	/** The list of party members to send the request for joining in progress. */
 	virtual TArray<UPartyMember*> GetLocalPartyMembersForJoinInProgress() const;
 
+	/** Override in child classes to provide extra invite metadata. */
+	virtual FString GetInviteMetadata(const FString& ExistingMetadata);
+
 	bool IsInviteRateLimited(const USocialUser& User, ESocialSubsystem SubsystemType) const;
 
 	bool ApplyCrossplayRestriction(FPartyJoinApproval& JoinApproval, const FUserPlatform& Platform, const FOnlinePartyData& JoinData) const;
@@ -321,7 +327,6 @@ protected:
 	/**
 	* Create a spectator beacon and connect to the server to get approval for new spectators
 	*/
-	void ConnectToSpectatorBeacon();
 	void CleanupSpectatorBeacon();
 	ASpectatorBeaconClient* CreateSpectatorBeaconClient();
 
diff --git a/Engine/Plugins/Online/OnlineFramework/Source/Party/Public/User/SocialUser.h b/Engine/Plugins/Online/OnlineFramework/Source/Party/Public/User/SocialUser.h
--- a/Engine/Plugins/Online/OnlineFramework/Source/Party/Public/User/SocialUser.h
+++ b/Engine/Plugins/Online/OnlineFramework/Source/Party/Public/User/SocialUser.h
@@ -190,16 +190,20 @@ protected:
 
 protected:
 	virtual void OnPresenceChangedInternal(ESocialSubsystem SubsystemType);
-	virtual void OnPartyInviteAcceptedInternal(const FOnlinePartyTypeId& PartyTypeId) const;
+	virtual void OnPartyInviteAcceptedInternal(const FOnlinePartyTypeId& PartyTypeId, const FString& InviteMetadata) const;
 	virtual void OnPartyInviteRejectedInternal(const FOnlinePartyTypeId& PartyTypeId) const;
 	virtual void HandleSetNicknameComplete(int32 LocalUserNum, const FUniqueNetId& FriendId, const FString& ListName, const FOnlineError& Error);
 	virtual void SetSubsystemId(ESocialSubsystem SubsystemType, const FUniqueNetIdRepl& SubsystemId);
 
+	UE_DEPRECATED(5.6, "Override 'void OnPartyInviteAcceptedInternal(const FOnlinePartyTypeId& PartyTypeId, const FString& InviteMetadata) const' instead")
+	virtual void OnPartyInviteAcceptedInternal(const FOnlinePartyTypeId& PartyTypeId) const {};
+
 	virtual void NotifyRequestToJoinReceived(const IOnlinePartyRequestToJoinInfo& Request) {}
 	virtual void NotifyRequestToJoinRemoved(const IOnlinePartyRequestToJoinInfo& Request, EPartyRequestToJoinRemovedReason Reason) {}
 	int32 NumPendingQueries = 0;
 	TMap<FString, FString> AnalyticsContext;
 
+	IOnlinePartyJoinInfoConstPtr GetSentPartyInvite(const FOnlinePartyTypeId& PartyTypeId) const;
 	void TryBroadcastInitializationComplete();
 
 	struct FSubsystemUserInfo
diff --git a/Engine/Source/Runtime/Analytics/Analytics/Public/AnalyticsConversion.h b/Engine/Source/Runtime/Analytics/Analytics/Public/AnalyticsConversion.h
--- a/Engine/Source/Runtime/Analytics/Analytics/Public/AnalyticsConversion.h
+++ b/Engine/Source/Runtime/Analytics/Analytics/Public/AnalyticsConversion.h
@@ -22,6 +22,11 @@ inline FString AnalyticsConversionToString(const FString& Value)
 	return Value;
 }
 
+inline FString AnalyticsConversionToString(FStringView Value)
+{
+	return FString(Value);
+}
+
 #define ANALYTICS_MAX_FLOAT_VALUE_BEFORE_EXPONENTIAL_NOTATION (1.0e+15F)
 
 inline FString AnalyticsConversionToString(float Value)
