Improve arguments passed down for analytics

diff --git a/Engine/Plugins/Online/OnlineFramework/Source/Party/Private/User/SocialUser.cpp b/Engine/Plugins/Online/OnlineFramework/Source/Party/Private/User/SocialUser.cpp
--- a/Engine/Plugins/Online/OnlineFramework/Source/Party/Private/User/SocialUser.cpp
+++ b/Engine/Plugins/Online/OnlineFramework/Source/Party/Private/User/SocialUser.cpp
@@ -1012,7 +1012,7 @@ void USocialUser::JoinParty(const FOnlinePartyTypeId& PartyTypeId, const FName&
 	{
 		IOnlinePartyPtr PartyInterface = Online::GetPartyInterfaceChecked(GetWorld());
 		PartyInterface->ClearInvitations(*GetOwningToolkit().GetLocalUserNetId(ESocialSubsystem::Primary), *GetUserId(ESocialSubsystem::Primary), nullptr);
-		OnPartyInviteAcceptedInternal(PartyTypeId, SentInvite->GetPlatformData());
+		OnPartyInviteAcceptedInternal(PartyTypeId, *SentInvite);
 	}
 }
 
@@ -1426,7 +1426,7 @@ void USocialUser::OnPresenceChangedInternal(ESocialSubsystem SubsystemType)
 	OnUserPresenceChanged().Broadcast(SubsystemType);
 }
 
-void USocialUser::OnPartyInviteAcceptedInternal(const FOnlinePartyTypeId& PartyTypeId, const FString& InviteMetadata) const
+void USocialUser::OnPartyInviteAcceptedInternal(const FOnlinePartyTypeId& PartyTypeId, const IOnlinePartyJoinInfo& Invite) const
 {
 	// Call deprecated method
 PRAGMA_DISABLE_DEPRECATION_WARNINGS
diff --git a/Engine/Plugins/Online/OnlineFramework/Source/Party/Public/User/SocialUser.h b/Engine/Plugins/Online/OnlineFramework/Source/Party/Public/User/SocialUser.h
--- a/Engine/Plugins/Online/OnlineFramework/Source/Party/Public/User/SocialUser.h
+++ b/Engine/Plugins/Online/OnlineFramework/Source/Party/Public/User/SocialUser.h
@@ -190,7 +190,7 @@ protected:
 
 protected:
 	virtual void OnPresenceChangedInternal(ESocialSubsystem SubsystemType);
-	virtual void OnPartyInviteAcceptedInternal(const FOnlinePartyTypeId& PartyTypeId, const FString& InviteMetadata) const;
+	virtual void OnPartyInviteAcceptedInternal(const FOnlinePartyTypeId& PartyTypeId, const IOnlinePartyJoinInfo& Invite) const;
 	virtual void OnPartyInviteRejectedInternal(const FOnlinePartyTypeId& PartyTypeId) const;
 	virtual void HandleSetNicknameComplete(int32 LocalUserNum, const FUniqueNetId& FriendId, const FString& ListName, const FOnlineError& Error);
 	virtual void SetSubsystemId(ESocialSubsystem SubsystemType, const FUniqueNetIdRepl& SubsystemId);
