Add string for TargetUserAway

diff --git a/Engine/Plugins/Online/OnlineFramework/Source/Party/Public/Party/PartyTypes.h b/Engine/Plugins/Online/OnlineFramework/Source/Party/Public/Party/PartyTypes.h
--- a/Engine/Plugins/Online/OnlineFramework/Source/Party/Public/Party/PartyTypes.h
+++ b/Engine/Plugins/Online/OnlineFramework/Source/Party/Public/Party/PartyTypes.h
@@ -149,82 +149,58 @@ inline const TCHAR* ToString(EPartyJoinDenialReason Type)
 	{
 	case EPartyJoinDenialReason::NoReason:
 		return TEXT("NoReason");
-		break;
 	case EPartyJoinDenialReason::JoinAttemptAborted:
 		return TEXT("JoinAttemptAborted");
-		break;
 	case EPartyJoinDenialReason::Busy:
 		return TEXT("Busy");
-		break;
 	case EPartyJoinDenialReason::OssUnavailable:
 		return TEXT("OssUnavailable");
-		break;
 	case EPartyJoinDenialReason::PartyFull:
 		return TEXT("PartyFull");
-		break;
 	case EPartyJoinDenialReason::GameFull:
 		return TEXT("GameFull");
-		break;
 	case EPartyJoinDenialReason::NotPartyLeader:
 		return TEXT("NotPartyLeader");
-		break;
 	case EPartyJoinDenialReason::PartyPrivate:
 		return TEXT("PartyPrivate");
-		break;
 	case EPartyJoinDenialReason::JoinerCrossplayRestricted:
 		return TEXT("JoinerCrossplayRestricted");
-		break;
 	case EPartyJoinDenialReason::MemberCrossplayRestricted:
 		return TEXT("MemberCrossplayRestricted");
-		break;
 	case EPartyJoinDenialReason::GameModeRestricted:
 		return TEXT("GameModeRestricted");
-		break;
 	case EPartyJoinDenialReason::Banned:
 		return TEXT("Banned");
-		break;
 	case EPartyJoinDenialReason::NotLoggedIn:
 		return TEXT("NotLoggedIn");
-		break;
 	case EPartyJoinDenialReason::CheckingForRejoin:
 		return TEXT("CheckingForRejoin");
-		break;
 	case EPartyJoinDenialReason::TargetUserMissingPresence:
 		return TEXT("TargetUserMissingPresence");
-		break;
 	case EPartyJoinDenialReason::TargetUserUnjoinable:
 		return TEXT("TargetUserUnjoinable");
-		break;
+	case EPartyJoinDenialReason::TargetUserAway:
+		return TEXT("TargetUserAway");
 	case EPartyJoinDenialReason::AlreadyLeaderInPlatformSession:
 		return TEXT("AlreadyLeaderInPlatformSession");
-		break;
 	case EPartyJoinDenialReason::TargetUserPlayingDifferentGame:
 		return TEXT("TargetUserPlayingDifferentGame");
-		break;
 	case EPartyJoinDenialReason::TargetUserMissingPlatformSession:
 		return TEXT("TargetUserMissingPlatformSession");
-		break;
 	case EPartyJoinDenialReason::PlatformSessionMissingJoinInfo:
 		return TEXT("PlatformSessionMissingJoinInfo");
-		break;
 	case EPartyJoinDenialReason::FailedToStartFindConsoleSession:
 		return TEXT("FailedToStartFindConsoleSession");
-		break;
 	case EPartyJoinDenialReason::MissingPartyClassForTypeId:
 		return TEXT("MissingPartyClassForTypeId");
-		break;
 	case EPartyJoinDenialReason::TargetUserBlocked:
 		return TEXT("TargetUserBlocked");
-		break;
 	case EPartyJoinDenialReason::InvalidJoinInfo:
 		return TEXT("InvalidJoinInfo");
-		break;
 	case EPartyJoinDenialReason::NotFriends:
 		return TEXT("NotFriends");
-		break;
 	default:
 		return TEXT("CustomReason");
-		break;
 	}
 }
 
