Add scope counter in OnPartyMemberExited events.
Remove defines that obfuscate party event names, to make them easier to discover

diff --git a/Engine/Plugins/Online/OnlineFramework/Source/Party/Private/Party/PartyPlatformSessionMonitor.cpp b/Engine/Plugins/Online/OnlineFramework/Source/Party/Private/Party/PartyPlatformSessionMonitor.cpp
--- a/Engine/Plugins/Online/OnlineFramework/Source/Party/Private/Party/PartyPlatformSessionMonitor.cpp
+++ b/Engine/Plugins/Online/OnlineFramework/Source/Party/Private/Party/PartyPlatformSessionMonitor.cpp
@@ -930,6 +930,7 @@ void FPartyPlatformSessionMonitor::HandlePartyMemberInitialized(UPartyMember* In
 
 void FPartyPlatformSessionMonitor::HandlePartyMemberLeft(UPartyMember* OldMember, const EMemberExitedReason ExitReason)
 {
+	QUICK_SCOPE_CYCLE_COUNTER(STAT_FPartyPlatformSessionMonitor_HandlePartyMemberLeft);
 	const FUniqueNetIdRepl PartyMemberPlatformUniqueId = OldMember->GetRepData().GetPlatformDataUniqueId();
 
 	if (PartyMemberPlatformUniqueId.IsValid())
@@ -955,7 +956,7 @@ void FPartyPlatformSessionMonitor::HandlePartyMemberLeft(UPartyMember* OldMember
 
 void FPartyPlatformSessionMonitor::HandleCreateSessionComplete(FName SessionName, bool bWasSuccessful)
 {
-	QUICK_SCOPE_CYCLE_COUNTER(STAT_USocialParty_HandleCreateSessionComplete);
+	QUICK_SCOPE_CYCLE_COUNTER(STAT_FPartyPlatformSessionMonitor_HandleCreateSessionComplete);
 	UE_LOG(LogParty, Log, TEXT("PartyPlatformSessionMonitor created platform session SessionName=[%s], bWasSuccessful=[%d]"), *SessionName.ToString(), bWasSuccessful);
 
 	const IOnlineSessionPtr& SessionInterface = SessionManager->GetSessionInterface();
diff --git a/Engine/Plugins/Online/OnlineFramework/Source/Party/Private/Party/SocialParty.cpp b/Engine/Plugins/Online/OnlineFramework/Source/Party/Private/Party/SocialParty.cpp
--- a/Engine/Plugins/Online/OnlineFramework/Source/Party/Private/Party/SocialParty.cpp
+++ b/Engine/Plugins/Online/OnlineFramework/Source/Party/Private/Party/SocialParty.cpp
@@ -1117,6 +1117,7 @@ void USocialParty::HandlePartyLeft(const FUniqueNetId& LocalUserId, const FOnlin
 
 void USocialParty::HandlePartyMemberExited(const FUniqueNetId& LocalUserId, const FOnlinePartyId& PartyId, const FUniqueNetId& MemberId, EMemberExitedReason ExitReason)
 {
+	QUICK_SCOPE_CYCLE_COUNTER(STAT_USocialParty_HandlePartyMemberExited);
 	if (PartyId == GetPartyId())
 	{
 		if (TObjectPtr<UPartyMember>* FoundPartyMember = PartyMembersById.Find(MemberId.AsShared()))
@@ -1136,7 +1137,10 @@ void USocialParty::HandlePartyMemberExited(const FUniqueNetId& LocalUserId, cons
 				UPartyMember& LeftMember = **FoundPartyMember;
 				PartyMembersById.Remove(FUniqueNetIdRepl(MemberId.AsShared()));
 
-				OnPartyMemberLeft().Broadcast(&LeftMember, ExitReason);
+				{
+					QUICK_SCOPE_CYCLE_COUNTER(STAT_USocialParty_HandlePartyMemberExited_Delegates);
+					OnPartyMemberLeft().Broadcast(&LeftMember, ExitReason);
+				}
 
 				TOptional<FString> SessionType = FPartyPlatformSessionManager::GetOssPartySessionType(LeftMember.GetPlatformOssName());
 				if (SessionType)
