Added CorrelationId field to Social.Party.RequestToJoin event
Added CorrelationId and RequestSentAt fields to Social.Party.AcceptRequestToJoin and Social.Party.DismissRequestToJoin

diff --git a/Engine/Plugins/Online/OnlineFramework/Source/Party/Private/User/SocialUser.cpp b/Engine/Plugins/Online/OnlineFramework/Source/Party/Private/User/SocialUser.cpp
--- a/Engine/Plugins/Online/OnlineFramework/Source/Party/Private/User/SocialUser.cpp
+++ b/Engine/Plugins/Online/OnlineFramework/Source/Party/Private/User/SocialUser.cpp
@@ -1198,8 +1198,13 @@ void USocialUser::RequestToJoinParty(const FName& JoinMethod)
 {
 	IOnlinePartyPtr PartyInterface = Online::GetPartyInterfaceChecked(GetWorld());
 	FPartyInvitationRecipient Recipient(*GetUserId(ESocialSubsystem::Primary));
-	Recipient.PlatformData = JoinMethod.ToString();
-	PartyInterface->RequestToJoinParty(*GetOwningToolkit().GetLocalUserNetId(ESocialSubsystem::Primary), IOnlinePartySystem::GetPrimaryPartyTypeId(), Recipient, FOnRequestToJoinPartyComplete::CreateUObject(this, &USocialUser::HandlePartyRequestToJoinSent, JoinMethod));
+	Recipient.PlatformData = GetRequestToJoinMetadata(JoinMethod.ToString());
+	PartyInterface->RequestToJoinParty(*GetOwningToolkit().GetLocalUserNetId(ESocialSubsystem::Primary), IOnlinePartySystem::GetPrimaryPartyTypeId(), Recipient, FOnRequestToJoinPartyComplete::CreateUObject(this, &USocialUser::HandlePartyRequestToJoinSent, JoinMethod, Recipient.PlatformData));
+}
+
+FString USocialUser::GetRequestToJoinMetadata(const FString& ExistingMetadata)
+{
+	return ExistingMetadata;
 }
 
 void USocialUser::AcceptRequestToJoinParty() const
@@ -1221,7 +1226,7 @@ void USocialUser::DismissRequestToJoinParty() const
 	}
 }
 
-void USocialUser::HandlePartyRequestToJoinSent(const FUniqueNetId& LocalUserId, const FUniqueNetId& PartyLeaderId, const FDateTime& ExpiresAt, const ERequestToJoinPartyCompletionResult Result, FName JoinMethod)
+void USocialUser::HandlePartyRequestToJoinSent(const FUniqueNetId& LocalUserId, const FUniqueNetId& PartyLeaderId, const FDateTime& ExpiresAt, const ERequestToJoinPartyCompletionResult Result, FName JoinMethod, FString Metadata)
 {
 	UE_LOG(LogParty, VeryVerbose, TEXT("%s - User [%s] sent a join request to [%s] with result[%s]"), ANSI_TO_TCHAR(__FUNCTION__), *LocalUserId.ToDebugString(), *PartyLeaderId.ToDebugString(), ToString(Result));
 
@@ -1229,6 +1234,13 @@ void USocialUser::HandlePartyRequestToJoinSent(const FUniqueNetId& LocalUserId,
 	{
 		GetOwningToolkit().OnPartyRequestToJoinSent().Broadcast(*this);
 	}
+PRAGMA_DISABLE_DEPRECATION_WARNINGS
+	HandlePartyRequestToJoinSent(LocalUserId, PartyLeaderId, ExpiresAt, Result, JoinMethod);
+PRAGMA_ENABLE_DEPRECATION_WARNINGS
+}
+
+void USocialUser::HandlePartyRequestToJoinSent(const FUniqueNetId& LocalUserId, const FUniqueNetId& PartyLeaderId, const FDateTime& ExpiresAt, const ERequestToJoinPartyCompletionResult Result, FName JoinMethod)
+{
 }
 
 TSharedPtr<const IOnlinePartyJoinInfo> USocialUser::GetPartyJoinInfo(const FOnlinePartyTypeId& PartyTypeId) const
diff --git a/Engine/Plugins/Online/OnlineFramework/Source/Party/Public/User/SocialUser.h b/Engine/Plugins/Online/OnlineFramework/Source/Party/Public/User/SocialUser.h
--- a/Engine/Plugins/Online/OnlineFramework/Source/Party/Public/User/SocialUser.h
+++ b/Engine/Plugins/Online/OnlineFramework/Source/Party/Public/User/SocialUser.h
@@ -121,8 +121,12 @@ public:
 	virtual void RequestToJoinParty(const FName& JoinMethod);
 	void AcceptRequestToJoinParty() const;
 	void DismissRequestToJoinParty() const;
+
+	virtual void HandlePartyRequestToJoinSent(const FUniqueNetId& LocalUserId, const FUniqueNetId& PartyLeaderId, const FDateTime& ExpiresAt, const ERequestToJoinPartyCompletionResult Result, FName JoinMethod, FString Metadata);
+
+	UE_DEPRECATED(5.6, "Use the overload above that also receives metadata")
 	virtual void HandlePartyRequestToJoinSent(const FUniqueNetId& LocalUserId, const FUniqueNetId& PartyLeaderId, const FDateTime& ExpiresAt, const ERequestToJoinPartyCompletionResult Result, FName JoinMethod);
-	
+
 	virtual IOnlinePartyJoinInfoConstPtr GetPartyJoinInfo(const FOnlinePartyTypeId& PartyTypeId) const;
 
 	bool HasSentPartyInvite(const FOnlinePartyTypeId& PartyTypeId) const;
@@ -198,6 +202,7 @@ protected:
 	UE_DEPRECATED(5.6, "Override 'void OnPartyInviteAcceptedInternal(const FOnlinePartyTypeId& PartyTypeId, const FString& InviteMetadata) const' instead")
 	virtual void OnPartyInviteAcceptedInternal(const FOnlinePartyTypeId& PartyTypeId) const {};
 
+	virtual FString GetRequestToJoinMetadata(const FString& ExistingMetadata);
 	virtual void NotifyRequestToJoinReceived(const IOnlinePartyRequestToJoinInfo& Request) {}
 	virtual void NotifyRequestToJoinRemoved(const IOnlinePartyRequestToJoinInfo& Request, EPartyRequestToJoinRemovedReason Reason) {}
 	int32 NumPendingQueries = 0;
