Patch check logging and analytics improvements.

diff --git a/Engine/Plugins/Online/OnlineFramework/Source/PatchCheck/Private/PatchCheck.cpp b/Engine/Plugins/Online/OnlineFramework/Source/PatchCheck/Private/PatchCheck.cpp
--- a/Engine/Plugins/Online/OnlineFramework/Source/PatchCheck/Private/PatchCheck.cpp
+++ b/Engine/Plugins/Online/OnlineFramework/Source/PatchCheck/Private/PatchCheck.cpp
@@ -84,10 +84,19 @@ void FPatchCheck::RefreshConfig()
 	}
 
 	GConfig->GetBool(TEXT("PatchCheck"), TEXT("bPlatformEnvironmentDetectionEnabled"), bPlatformEnvironmentDetectionEnabled, GEngineIni);
+
+	UE_LOG(LogPatchCheck, VeryVerbose, TEXT("[%hs] [bCheckPlatformOSSForUpdate=%s], [bCheckOSSForUpdate=%s], [bPlatformEnvironmentDetectionEnabled=%s]"),
+		__FUNCTION__, *LexToString(bCheckPlatformOSSForUpdate), *LexToString(bCheckOSSForUpdate), *LexToString(bPlatformEnvironmentDetectionEnabled));
 }
 
 void FPatchCheck::OnDetectPlatformEnvironmentComplete(const FOnlineError& Result)
 {
+	UE_LOG(LogPatchCheck, VeryVerbose, TEXT("[%hs] [Result=%s]"), __FUNCTION__, *Result.ToLogString());
+	if (TSharedPtr<IPatchCheckStatsCollector> StrongStats = Stats.Pin())
+	{
+		StrongStats->OnPatchCheckStep_DetectEnvironmentComplete(Result.WasSuccessful(), Result.ToLogString());
+	}
+
 	if (Result.WasSuccessful())
 	{
 		bPlatformEnvironmentDetected = true;
@@ -97,20 +106,20 @@ void FPatchCheck::OnDetectPlatformEnvironmentComplete(const FOnlineError& Result
 	{
 		if (Result.GetErrorCode().Contains(TEXT("getUserAccessCode failed : 0x8055000f"), ESearchCase::IgnoreCase))
 		{
-			UE_LOG(LogPatchCheck, Warning, TEXT("Failed to complete login because patch is required"));
+			UE_LOG(LogPatchCheck, Warning, TEXT("[%hs] Failed to complete login because patch is required"), __FUNCTION__);
 			PatchCheckComplete(EPatchCheckResult::PatchRequired);
 		}
 		else
 		{
 			if (Result.GetErrorCode().Contains(TEXT("com.epicgames.identity.notloggedin"), ESearchCase::IgnoreCase))
 			{
-				UE_LOG(LogPatchCheck, Warning, TEXT("Failed to detect online environment for the platform, no user signed in"));
+				UE_LOG(LogPatchCheck, Warning, TEXT("[%hs] Failed to detect online environment for the platform, no user signed in"), __FUNCTION__);
 				PatchCheckComplete(EPatchCheckResult::NoLoggedInUser);
 			}
 			else
 			{
 				// just a platform env error, assume production and keep going
-				UE_LOG(LogPatchCheck, Warning, TEXT("Failed to detect online environment for the platform"));
+				UE_LOG(LogPatchCheck, Warning, TEXT("[%hs] Failed to detect online environment for the platform"), __FUNCTION__);
 				bPlatformEnvironmentDetected = true;
 				HandleOSSPatchCheck();
 			}
@@ -120,13 +129,26 @@ void FPatchCheck::OnDetectPlatformEnvironmentComplete(const FOnlineError& Result
 
 void FPatchCheck::StartPatchCheck()
 {
+	UE_LOG(LogPatchCheck, VeryVerbose, TEXT("[%hs] [bIsCheckInProgress=%s], [bPlatformEnvironmentDetectionEnabled=%s], [bPlatformEnvironmentDetected=%s]"),
+		__FUNCTION__, *LexToString(bIsCheckInProgress), *LexToString(bPlatformEnvironmentDetectionEnabled), *LexToString(bPlatformEnvironmentDetected));
 	if (bIsCheckInProgress)
 		return;
 
 	RefreshConfig();
 
+	TSharedPtr<IPatchCheckStatsCollector> StrongStats = Stats.Pin();
+	if (StrongStats)
+	{
+		StrongStats->OnPatchCheckStarted();
+	}
+
 	if (bPlatformEnvironmentDetectionEnabled && !bPlatformEnvironmentDetected)
 	{
+		if (StrongStats)
+		{
+			StrongStats->OnPatchCheckStep_DetectEnvironmentStarted();
+		}
+
 #if PATCH_CHECK_PLATFORM_ENVIRONMENT_DETECTION
 		if (DetectPlatformEnvironment())
 		{
@@ -140,6 +162,7 @@ void FPatchCheck::StartPatchCheck()
 
 void FPatchCheck::HandleOSSPatchCheck()
 {
+	UE_LOG(LogPatchCheck, VeryVerbose, TEXT("[%hs]"), __FUNCTION__);
 	if (bCheckPlatformOSSForUpdate && IOnlineSubsystem::GetByPlatform() != nullptr)
 	{
 		bIsCheckInProgress = true;
@@ -166,8 +189,19 @@ void FPatchCheck::RemoveEnvironmentWantsPatchCheckBackCompatDelegate(FName Tag)
 	BackCompatEnvironmentWantsPatchCheckDelegates.Remove(Tag);
 }
 
+void FPatchCheck::RegisterStatsCollector(const TSharedPtr<IPatchCheckStatsCollector>& InStats)
+{
+	Stats = InStats;
+}
+
 void FPatchCheck::StartPlatformOSSPatchCheck()
 {
+	UE_LOG(LogPatchCheck, VeryVerbose, TEXT("[%hs]"), __FUNCTION__);
+	if (TSharedPtr<IPatchCheckStatsCollector> StrongStats = Stats.Pin())
+	{
+		StrongStats->OnPatchCheckStep_CheckPlatformPatchStarted();
+	}
+
 	EPatchCheckResult PatchResult = EPatchCheckResult::PatchCheckFailure;
 	bool bStarted = false;
 
@@ -205,10 +239,15 @@ void FPatchCheck::StartPlatformOSSPatchCheck()
 
 void FPatchCheck::StartOSSPatchCheck()
 {
-	if (SkipPatchCheck())
+	const bool bSkipPatchCheck = SkipPatchCheck();
+	UE_LOG(LogPatchCheck, Verbose, TEXT("[%hs] [bSkipPatchCheck=%s]"), __FUNCTION__, *LexToString(bSkipPatchCheck));
+	if (TSharedPtr<IPatchCheckStatsCollector> StrongStats = Stats.Pin())
 	{
-		UE_LOG(LogPatchCheck, Verbose, TEXT("[StartOSSPatchCheck] Skipping patch check."));
+		StrongStats->OnPatchCheckStep_CheckOnlineServicePatchStarted();
+	}
 
+	if (bSkipPatchCheck)
+	{
 		// Trigger completion if check is skipped.
 		PatchCheckComplete(EPatchCheckResult::NoPatchRequired);
 	}
@@ -283,6 +322,9 @@ bool FPatchCheck::SkipPatchCheck() const
 	const bool bForceSkipCheck = FParse::Param(FCommandLine::Get(), TEXT("SkipPatchCheck"));
 	const bool bSkipPatchCheck = !bForcePatchCheck && (!bEnvironmentWantsPatchCheck || bSkipDueToEditor || bSkipDueToDedicatedServer || bForceSkipCheck || bSkipDueToUnattended);
 
+	UE_LOG(LogPatchCheck, VeryVerbose, TEXT("[%hs] [bSkipPatchCheck=%s], [bForcePatchCheck=%s], [bEnvironmentWantsPatchCheck=%s], [bSkipDueToEditor=%s], [bSkipDueToDedicatedServer=%s], [bForceSkipCheck=%s], [bSkipDueToUnattended=%s]"),
+		__FUNCTION__, *LexToString(bSkipPatchCheck), *LexToString(bForcePatchCheck), *LexToString(bEnvironmentWantsPatchCheck), *LexToString(bSkipDueToEditor), *LexToString(bSkipDueToDedicatedServer), *LexToString(bForceSkipCheck), *LexToString(bSkipDueToUnattended));
+
 	return bSkipPatchCheck;
 }
 
@@ -321,7 +363,21 @@ inline EPatchCheckResult TranslatePatchCheckResult(uint32 PrivilegeResult)
 void FPatchCheck::OnCheckForPatchComplete(const FUniqueNetId& UniqueId, EUserPrivileges::Type Privilege, uint32 PrivilegeResult, bool bConsoleCheck)
 {
 	const EPatchCheckResult Result = Privilege == EUserPrivileges::CanPlayOnline ? TranslatePatchCheckResult(PrivilegeResult) : EPatchCheckResult::NoPatchRequired;
-	UE_LOG(LogPatchCheck, Verbose, TEXT("[OnCheckForPatchComplete] Type: %s, Privilege: %d, PrivilegeResult: %d, PatchCheckResult: %s"), bConsoleCheck ? TEXT("PlatformOSS") : TEXT("DefaultOSS"), (uint32)Privilege, PrivilegeResult, LexToString(Result));
+	UE_LOG(LogPatchCheck, Verbose, TEXT("[%hs] [Type=%s], [Privilege=%s], [PrivilegeResult=%s], [PrivilegeResultValue=%d], [PatchCheckResult=%s]"),
+		__FUNCTION__, bConsoleCheck ? TEXT("PlatformOSS") : TEXT("DefaultOSS"), *ToDebugString(Privilege), *ToDebugString(static_cast<IOnlineIdentity::EPrivilegeResults>(PrivilegeResult)), PrivilegeResult, LexToString(Result));
+
+	if (TSharedPtr<IPatchCheckStatsCollector> StrongStats = Stats.Pin())
+	{
+		const bool bSucceeded = Result == EPatchCheckResult::NoPatchRequired;
+		if (bConsoleCheck)
+		{
+			StrongStats->OnPatchCheckStep_CheckPlatformPatchComplete(bSucceeded, LexToString(Result));
+		}
+		else
+		{
+			StrongStats->OnPatchCheckStep_CheckOnlineServicePatchComplete(bSucceeded, LexToString(Result));
+		}
+	}
 
 	if (bCheckOSSForUpdate && bConsoleCheck && Result == EPatchCheckResult::NoPatchRequired)
 	{
@@ -335,6 +391,12 @@ void FPatchCheck::OnCheckForPatchComplete(const FUniqueNetId& UniqueId, EUserPri
 
 void FPatchCheck::PatchCheckComplete(EPatchCheckResult PatchResult)
 {
+	UE_LOG(LogPatchCheck, Log, TEXT("[%hs] [PatchResult=%s]"), __FUNCTION__, LexToString(PatchResult));
+	if (TSharedPtr<IPatchCheckStatsCollector> StrongStats = Stats.Pin())
+	{
+		StrongStats->OnPatchCheckComplete(PatchResult);
+	}
+
 	OnComplete.Broadcast(PatchResult);
 	bIsCheckInProgress = false;
 }
diff --git a/Engine/Plugins/Online/OnlineFramework/Source/PatchCheck/Public/PatchCheck.h b/Engine/Plugins/Online/OnlineFramework/Source/PatchCheck/Public/PatchCheck.h
--- a/Engine/Plugins/Online/OnlineFramework/Source/PatchCheck/Public/PatchCheck.h
+++ b/Engine/Plugins/Online/OnlineFramework/Source/PatchCheck/Public/PatchCheck.h
@@ -25,6 +25,24 @@ enum class EPatchCheckResult : uint8
 };
 PATCHCHECK_API const TCHAR* LexToString(EPatchCheckResult Value);
 
+class IPatchCheckStatsCollector
+{
+public:
+	virtual ~IPatchCheckStatsCollector() = default;
+
+	// Overall process/
+	virtual void OnPatchCheckStarted() = 0;
+	virtual void OnPatchCheckComplete(EPatchCheckResult Result) = 0;
+
+	// Individual stages.
+	virtual void OnPatchCheckStep_DetectEnvironmentStarted() = 0;
+	virtual void OnPatchCheckStep_DetectEnvironmentComplete(bool bSuccess, const FString& Error) = 0;
+	virtual void OnPatchCheckStep_CheckPlatformPatchStarted() = 0;
+	virtual void OnPatchCheckStep_CheckPlatformPatchComplete(bool bSuccess, const FString& Error) = 0;
+	virtual void OnPatchCheckStep_CheckOnlineServicePatchStarted() = 0;
+	virtual void OnPatchCheckStep_CheckOnlineServicePatchComplete(bool bSuccess, const FString& Error) = 0;
+};
+
 DECLARE_MULTICAST_DELEGATE_OneParam(FOnPatchCheckComplete, EPatchCheckResult /*Result*/);
 
 // For backwards compatibility only!
@@ -53,6 +71,8 @@ public:
 
 	FOnPatchCheckComplete& GetOnComplete() { return OnComplete; }
 
+	void RegisterStatsCollector(const TSharedPtr<IPatchCheckStatsCollector>& InStats);
+
 protected:
 	void RefreshConfig();
 	virtual void StartPlatformOSSPatchCheck();
@@ -83,6 +103,8 @@ protected:
 
 protected:
 
+	TWeakPtr<IPatchCheckStatsCollector> Stats;
+
 	FOnPatchCheckComplete OnComplete;
 
 	/** For backwards compatibility with UUpdateManager */
