[PatchCheck] Fail open on patch check error.

diff --git a/Engine/Plugins/Online/OnlineFramework/Source/PatchCheck/PatchCheck.build.cs b/Engine/Plugins/Online/OnlineFramework/Source/PatchCheck/PatchCheck.build.cs
--- a/Engine/Plugins/Online/OnlineFramework/Source/PatchCheck/PatchCheck.build.cs
+++ b/Engine/Plugins/Online/OnlineFramework/Source/PatchCheck/PatchCheck.build.cs
@@ -1,5 +1,6 @@
 // Copyright Epic Games, Inc. All Rights Reserved.
 
+using System;
 using System.IO;
 using UnrealBuildTool;
 
@@ -22,11 +23,11 @@ public class PatchCheck : ModuleRules
 			}
 		);
 
-		PrivateDefinitions.Add("PATCH_CHECK_FAIL_ON_GENERIC_FAILURE=" + (bFailOnGenericFailure ? "1" : "0"));
 		PublicDefinitions.Add("PATCH_CHECK_PLATFORM_ENVIRONMENT_DETECTION=" + (bPlatformEnvironmentDetection ? "1" : "0"));
 		PrivateDefinitions.Add("PATCH_CHECK_PRIVILEGE_MUST_BE_LOGGED_IN=" + (bCanCheckPlayOnlinePrivilegeMustBeLoggedIn ? "1" : "0"));
 	}
 
+	[Obsolete("bFailOnGenericFailure deprecated since UE 5.6 - Patch checks always fail open when there are service issues.")]
 	protected virtual bool bFailOnGenericFailure { get { return true; } }
 	protected virtual bool bPlatformEnvironmentDetection { get { return false; } }
 	protected virtual bool bCanCheckPlayOnlinePrivilegeMustBeLoggedIn { get { return true; } }
diff --git a/Engine/Plugins/Online/OnlineFramework/Source/PatchCheck/Private/PatchCheck.cpp b/Engine/Plugins/Online/OnlineFramework/Source/PatchCheck/Private/PatchCheck.cpp
--- a/Engine/Plugins/Online/OnlineFramework/Source/PatchCheck/Private/PatchCheck.cpp
+++ b/Engine/Plugins/Online/OnlineFramework/Source/PatchCheck/Private/PatchCheck.cpp
@@ -3,11 +3,17 @@
 #include "PatchCheck.h"
 #include "PatchCheckModule.h"
 
+#include "HAL/IConsoleManager.h"
 #include "Misc/App.h"
 #include "Online.h"
 
 DEFINE_LOG_CATEGORY(LogPatchCheck);
 
+static TAutoConsoleVariable<bool> CVarPatchCheckFailOpenOnError(
+	TEXT("PatchCheck.FailOpenOnError"),
+	true,
+	TEXT("Whether a service failure when running a patch check should fail the operation."));
+
 const TCHAR* LexToString(EPatchCheckResult Value)
 {
 	switch (Value)
@@ -346,15 +352,7 @@ inline EPatchCheckResult TranslatePatchCheckResult(uint32 PrivilegeResult)
 	}
 	else if (PrivilegeResult & (uint32)IOnlineIdentity::EPrivilegeResults::GenericFailure)
 	{
-		CA_CONSTANT_IF(PATCH_CHECK_FAIL_ON_GENERIC_FAILURE)
-		{
-			Result = EPatchCheckResult::PatchCheckFailure;
-		}
-		else
-		{
-			// Skip console backend failures
-			Result = EPatchCheckResult::NoPatchRequired;
-		}
+		Result = EPatchCheckResult::PatchCheckFailure;
 	}
 
 	return Result;
@@ -362,10 +360,30 @@ inline EPatchCheckResult TranslatePatchCheckResult(uint32 PrivilegeResult)
 
 void FPatchCheck::OnCheckForPatchComplete(const FUniqueNetId& UniqueId, EUserPrivileges::Type Privilege, uint32 PrivilegeResult, bool bConsoleCheck)
 {
-	const EPatchCheckResult Result = Privilege == EUserPrivileges::CanPlayOnline ? TranslatePatchCheckResult(PrivilegeResult) : EPatchCheckResult::NoPatchRequired;
+	EPatchCheckResult Result = Privilege == EUserPrivileges::CanPlayOnline ? TranslatePatchCheckResult(PrivilegeResult) : EPatchCheckResult::NoPatchRequired;
+
+#if !UE_BUILD_SHIPPING
+	// Set failure if requested.
+	static bool bPatchCheckMockFailure = FParse::Param(FCommandLine::Get(), TEXT("PatchCheckMockFailure"));
+	if (bPatchCheckMockFailure)
+	{
+		UE_LOG(LogPatchCheck, Log, TEXT("[%hs] Simulating patch check failure."), __FUNCTION__);
+		Result = EPatchCheckResult::PatchCheckFailure;
+	}
+
+	// Set patch required if requested.
+	static bool bPatchCheckMockPatchRequired = FParse::Param(FCommandLine::Get(), TEXT("PatchCheckMockPatchRequired"));
+	if (bPatchCheckMockPatchRequired)
+	{
+		UE_LOG(LogPatchCheck, Log, TEXT("[%hs] Simulating required patch available."), __FUNCTION__);
+		Result = EPatchCheckResult::PatchRequired;
+	}
+#endif
+
 	UE_LOG(LogPatchCheck, Verbose, TEXT("[%hs] [Type=%s], [Privilege=%s], [PrivilegeResult=%s], [PrivilegeResultValue=%d], [PatchCheckResult=%s]"),
 		__FUNCTION__, bConsoleCheck ? TEXT("PlatformOSS") : TEXT("DefaultOSS"), *ToDebugString(Privilege), *ToDebugString(static_cast<IOnlineIdentity::EPrivilegeResults>(PrivilegeResult)), PrivilegeResult, LexToString(Result));
 
+	// Publish stats.
 	if (TSharedPtr<IPatchCheckStatsCollector> StrongStats = Stats.Pin())
 	{
 		const bool bSucceeded = Result == EPatchCheckResult::NoPatchRequired;
@@ -379,6 +397,28 @@ void FPatchCheck::OnCheckForPatchComplete(const FUniqueNetId& UniqueId, EUserPri
 		}
 	}
 
+	// If the result is a failure, check whether the result should be remapped to NoPatchRequired.
+	TOptional<FPreviousSuccessInfo>& PreviousCachedSuccess = bConsoleCheck ? PreviousPlatformSuccess : PreviousOnlineServiceSuccess;
+	if (Result == EPatchCheckResult::PatchCheckFailure)
+	{
+		if (CVarPatchCheckFailOpenOnError.GetValueOnGameThread())
+		{
+			auto GetLastSuccessString = [](const TOptional<FPreviousSuccessInfo>& CachedSuccess) -> FString
+			{
+				return CachedSuccess.IsSet() ? FString::Printf(TEXT("%f seconds ago"), FPlatformTime::Seconds() - CachedSuccess->ResultTime) : TEXT("Never");
+			};
+			UE_LOG(LogPatchCheck, Log, TEXT("[%hs] Remapping failure to NoPatchRequired. [Type=%s], [LastSuccess=%s]"),
+				__FUNCTION__, bConsoleCheck ? TEXT("PlatformOSS") : TEXT("DefaultOSS"), *GetLastSuccessString(PreviousCachedSuccess));
+			Result = EPatchCheckResult::NoPatchRequired;
+		}
+	}
+	else if (Result == EPatchCheckResult::NoPatchRequired)
+	{
+		// Store the most recent success.
+		FPreviousSuccessInfo& SuccessInfo = PreviousCachedSuccess.IsSet() ? *PreviousCachedSuccess : PreviousCachedSuccess.Emplace();
+		SuccessInfo.ResultTime = FPlatformTime::Seconds();
+	}
+
 	if (bCheckOSSForUpdate && bConsoleCheck && Result == EPatchCheckResult::NoPatchRequired)
 	{
 		// We perform both checks in this case
diff --git a/Engine/Plugins/Online/OnlineFramework/Source/PatchCheck/Public/PatchCheck.h b/Engine/Plugins/Online/OnlineFramework/Source/PatchCheck/Public/PatchCheck.h
--- a/Engine/Plugins/Online/OnlineFramework/Source/PatchCheck/Public/PatchCheck.h
+++ b/Engine/Plugins/Online/OnlineFramework/Source/PatchCheck/Public/PatchCheck.h
@@ -103,7 +103,14 @@ protected:
 
 protected:
 
+	struct FPreviousSuccessInfo
+	{
+		double ResultTime = 0;
+	};
+
 	TWeakPtr<IPatchCheckStatsCollector> Stats;
+	TOptional<FPreviousSuccessInfo> PreviousPlatformSuccess;
+	TOptional<FPreviousSuccessInfo> PreviousOnlineServiceSuccess;
 
 	FOnPatchCheckComplete OnComplete;
 
