[AutoRTFM] Enable AutoRTFM in the online framework party system.

This code includes low-level classes like FUserPlatform, which can be reasonably accessed from within a Verse transaction. For instance, PlatformRequirementGate uses this to allow for different behaviors between mobile, console or desktop.

diff --git a/Engine/Plugins/Online/OnlineFramework/Source/Party/Party.Build.cs b/Engine/Plugins/Online/OnlineFramework/Source/Party/Party.Build.cs
--- a/Engine/Plugins/Online/OnlineFramework/Source/Party/Party.Build.cs
+++ b/Engine/Plugins/Online/OnlineFramework/Source/Party/Party.Build.cs
@@ -27,15 +27,6 @@ public class Party : ModuleRules
 				"Engine",
 			}
 			);
-
-		if (Target.Platform == UnrealTargetPlatform.Win64)
-		{
-		}
-		else if (Target.Platform == UnrealTargetPlatform.Mac)
-		{
-		}
-
-		bDisableAutoRTFMInstrumentation = true;
 	}
 
 	protected virtual bool bUsesPSNSessions { get { return false; } }
