Overhaul GFP hotfixing

- Fix an issue where a hotfix CVar would be applied as PluginLowPriority
- Fix an issue where a duplicate hotfix Cvar would not be reverted fully on unload
- Fix ignored hotfix prefixes for a branch or CL
- Fix platform pattern using "Platform" instead of "Platform_"

#rnx

diff --git a/Engine/Plugins/Online/OnlineFramework/Source/Hotfix/Private/OnlineHotfixManager.cpp b/Engine/Plugins/Online/OnlineFramework/Source/Hotfix/Private/OnlineHotfixManager.cpp
--- a/Engine/Plugins/Online/OnlineFramework/Source/Hotfix/Private/OnlineHotfixManager.cpp
+++ b/Engine/Plugins/Online/OnlineFramework/Source/Hotfix/Private/OnlineHotfixManager.cpp
@@ -297,6 +297,8 @@ void UOnlineHotfixManager::StartHotfixProcess()
 {
 	UE_LOG(LogHotfixManager, Log, TEXT("Starting Hotfix Process"));
 
+	DynamicHotfixContents.Empty();
+
 	// Patching the editor this way seems like a bad idea
 	const bool bShouldHotfix = ShouldPerformHotfix();
 	if (!bShouldHotfix)
@@ -1010,8 +1012,12 @@ bool UOnlineHotfixManager::HotfixIniFile(const FString& FileName, const FString&
 			}
 			else
 			{
-				UE_LOG(LogHotfixManager, Log, TEXT("Storing ini data for pending hotfix file %s"), *FileName);
-				DynamicHotfixContents.Add(*FileName, IniData);
+				const FString HotfixTag = GetStrippedConfigFileName(FileName).Replace(TEXT(".ini"), TEXT(""));
+
+				UE_LOG(LogHotfixManager, Log, TEXT("Storing ini data for pending hotfix %s (%s)"), *FileName, *HotfixTag);
+
+				TArray<TPair<FString, FString>>& HotfixesForTag = DynamicHotfixContents.FindOrAdd(*HotfixTag);
+				HotfixesForTag.Add(TPair<FString, FString>(FileName, IniData));
 			}
 		});
 		
@@ -1989,31 +1995,21 @@ void UOnlineHotfixManager::StopTrackingInvalidHotfixedAssets()
 
 void UOnlineHotfixManager::HotfixDynamicBranch(const FName& Tag, const FName& Branch, class FConfigModificationTracker* ModificationTracker)
 {
-	FConfigBranch* BranchToHotfix = GConfig->FindBranch(Branch, FString());
-	if (BranchToHotfix)
+	const FString HotfixBaseName = Tag.ToString() + Branch.ToString();
+	const TArray<TPair<FString, FString>>* HotfixesForTag = DynamicHotfixContents.Find(*HotfixBaseName);
+
+	if (HotfixesForTag)
 	{
-		// Check in dynamic hotfix contents that we can find a given ini file and apply it as a dynamic layer
-		const auto TryApplyHotfix = [&](const FString Prefix)
+		for (const TPair<FString, FString>& HotfixFilenameAndContents : *HotfixesForTag)
 		{
-			const FString HotfixFileName = Prefix + Tag.ToString() + Branch.ToString() + TEXT(".ini");
-			const FString* IniContents = DynamicHotfixContents.Find(*HotfixFileName);
+			UE_LOG(LogHotfixManager, Log, TEXT("HotfixDynamicBranch: applying hotfix %s for tag %s on branch %s"), *HotfixFilenameAndContents.Key, *Tag.ToString(), *Branch.ToString());
 
-			if (IniContents)
+			FConfigBranch* BranchToHotfix = GConfig->FindBranch(Branch, FString());
+			if (BranchToHotfix && !BranchToHotfix->AddDynamicLayerStringToHierarchy(HotfixFilenameAndContents.Key, HotfixFilenameAndContents.Value, Tag, DynamicLayerPriority::Hotfix, ModificationTracker))
 			{
-				UE_LOG(LogHotfixManager, Log, TEXT("HotfixDynamicBranch: applying hotfix %s for tag %s on branch %s"), *HotfixFileName, *Tag.ToString(), *Branch.ToString());
-				if (!BranchToHotfix->AddDynamicLayerStringToHierarchy(HotfixFileName, *IniContents, Tag, DynamicLayerPriority::Hotfix, ModificationTracker))
-				{
-					UE_LOG(LogHotfixManager, Warning, TEXT("HotfixDynamicBranch: failed to apply hotfix"));
-				}
+				UE_LOG(LogHotfixManager, Warning, TEXT("HotfixDynamicBranch: failed to apply hotfix"));
 			}
-		};
-
-		const FString PlatformName = FPlatformProperties::IniPlatformName();
-
-		// Valid hotfix filenames: $Plugin$Branch.ini, Default$Plugin$Branch.ini, $Platform$PluginBranch.ini
-		TryApplyHotfix(FString());
-		TryApplyHotfix("Default");
-		TryApplyHotfix(PlatformName);
+		}
 	}
 }
 
diff --git a/Engine/Plugins/Online/OnlineFramework/Source/Hotfix/Public/OnlineHotfixManager.h b/Engine/Plugins/Online/OnlineFramework/Source/Hotfix/Public/OnlineHotfixManager.h
--- a/Engine/Plugins/Online/OnlineFramework/Source/Hotfix/Public/OnlineHotfixManager.h
+++ b/Engine/Plugins/Online/OnlineFramework/Source/Hotfix/Public/OnlineHotfixManager.h
@@ -197,7 +197,7 @@ protected:
 	/** Our passed-in World */
 	TWeakObjectPtr<UWorld> OwnerWorld;
 	/** Loaded hotfix contents that were not mapped to any known branch, but might be loaded later */
-	TMap<FName, FString> DynamicHotfixContents;
+	TMap<FName, TArray<TPair<FString, FString>>> DynamicHotfixContents;
 
 	virtual void Init();
 	virtual void Cleanup();
diff --git a/Engine/Plugins/Runtime/GameFeatures/Source/GameFeatures/Private/GameFeatureData.cpp b/Engine/Plugins/Runtime/GameFeatures/Source/GameFeatures/Private/GameFeatureData.cpp
--- a/Engine/Plugins/Runtime/GameFeatures/Source/GameFeatures/Private/GameFeatureData.cpp
+++ b/Engine/Plugins/Runtime/GameFeatures/Source/GameFeatures/Private/GameFeatureData.cpp
@@ -191,9 +191,16 @@ void UGameFeatureData::InitializeHierarchicalPluginIniFiles(const FString& Plugi
 
 				// apply plugin modifications from this plugin to the everything
 				FConfigCacheIni::AddPluginToAllBranches(PluginName, ChangeTracker);
+			});
 
-				// give hotfix a chance to run at that point
+		UE::DynamicConfig::PerformDynamicConfig(PluginName, [PluginName](FConfigModificationTracker* ChangeTracker)
+			{
 				{
+					// set which sections to track for cvars, with their priority
+					ChangeTracker->CVars.Add(TEXT("ConsoleVariables")).CVarPriority = (int)ECVF_SetByHotfix;
+					ChangeTracker->CVars.Add(TEXT("ConsoleVariables_HighPriority")).CVarPriority = (int)ECVF_SetByHotfix;
+
+					// give hotfix a chance to run at that point
 #define HOTFIX_BRANCH(Ini) UE::DynamicConfig::HotfixPluginForBranch.Broadcast(PluginName, #Ini, ChangeTracker);
 					ENUMERATE_KNOWN_INI_FILES(HOTFIX_BRANCH);
 #undef HOTFIX_BRANCH
diff --git a/Engine/Source/Runtime/Core/Private/HAL/ConsoleManager.cpp b/Engine/Source/Runtime/Core/Private/HAL/ConsoleManager.cpp
--- a/Engine/Source/Runtime/Core/Private/HAL/ConsoleManager.cpp
+++ b/Engine/Source/Runtime/Core/Private/HAL/ConsoleManager.cpp
@@ -940,7 +940,6 @@ public:
 					{
 						bUnsetSomething = true;
 						It.RemoveCurrent();
-						break;
 					}
 				}
 
diff --git a/Engine/Source/Runtime/Core/Private/Misc/ConfigCacheIni.cpp b/Engine/Source/Runtime/Core/Private/Misc/ConfigCacheIni.cpp
--- a/Engine/Source/Runtime/Core/Private/Misc/ConfigCacheIni.cpp
+++ b/Engine/Source/Runtime/Core/Private/Misc/ConfigCacheIni.cpp
@@ -3334,7 +3334,8 @@ bool FConfigBranch::AddDynamicLayerStringToHierarchy(const FString& Filename, co
 					FConfigSection NewSection;
 					// copy just the SectionMap parts
 					(FConfigSectionMap&)NewSection = (const FConfigSectionMap&)Pair.Value;
-					CVarTracker->CVarEntriesPerBranchPerTag.FindOrAdd(Tag).Add(IniName, NewSection);
+					FConfigSection& SectionEntry = CVarTracker->CVarEntriesPerBranchPerTag.FindOrAdd(Tag).FindOrAdd(IniName);
+					SectionEntry.Append(NewSection);
 				}
 			}
 		}
