Allow to add rows to datatables with a hotfix

#rnx

diff --git a/Engine/Plugins/Online/OnlineFramework/Source/Hotfix/Hotfix.Build.cs b/Engine/Plugins/Online/OnlineFramework/Source/Hotfix/Hotfix.Build.cs
--- a/Engine/Plugins/Online/OnlineFramework/Source/Hotfix/Hotfix.Build.cs
+++ b/Engine/Plugins/Online/OnlineFramework/Source/Hotfix/Hotfix.Build.cs
@@ -15,7 +15,8 @@ public class Hotfix : ModuleRules
 				"CoreUObject",
 				"Engine",
 				"MoviePlayerProxy",
-				"OnlineSubsystemUtils"
+				"OnlineSubsystemUtils",
+				"Json",
 			}
 			);
 
diff --git a/Engine/Plugins/Online/OnlineFramework/Source/Hotfix/Private/OnlineHotfixManager.cpp b/Engine/Plugins/Online/OnlineFramework/Source/Hotfix/Private/OnlineHotfixManager.cpp
--- a/Engine/Plugins/Online/OnlineFramework/Source/Hotfix/Private/OnlineHotfixManager.cpp
+++ b/Engine/Plugins/Online/OnlineFramework/Source/Hotfix/Private/OnlineHotfixManager.cpp
@@ -1449,46 +1449,89 @@ void UOnlineHotfixManager::PatchAssetsFromIniFiles()
 						if (Asset != nullptr)
 						{
 							const FString RowUpdate(TEXT("RowUpdate"));
+							const FString AddRow(TEXT("AddRow"));
 							const FString TableUpdate(TEXT("TableUpdate"));
 							const FString CurveUpdate(TEXT("CurveUpdate"));
 
-							if (HotfixType == RowUpdate && Tokens.Num() == 5)
+							if (HotfixType == RowUpdate)
 							{
-								// The hotfix line should be
-								//	+DataTable=<data table path>;RowUpdate;<row name>;<column name>;<new value>
-								//	+CurveTable=<curve table path>;RowUpdate;<row name>;<column name>;<new value>
-								//	+CurveFloat=<curve float path>;RowUpdate;None;<column name>;<new value>
-								HotfixRowUpdate(Asset, AssetPath, Tokens[2], Tokens[3], Tokens[4], ProblemStrings, &ChangedDataTables, ChangedCurveTablesPointer);
-								bAddAssetToHotfixedList = ProblemStrings.Num() == 0;
+								if (Tokens.Num() == 5)
+								{
+									// The hotfix line should be
+									//	+DataTable=<data table path>;RowUpdate;<row name>;<column name>;<new value>
+									//	+CurveTable=<curve table path>;RowUpdate;<row name>;<column name>;<new value>
+									//	+CurveFloat=<curve float path>;RowUpdate;None;<column name>;<new value>
+									HotfixRowUpdate(Asset, AssetPath, Tokens[2], Tokens[3], Tokens[4], ProblemStrings, &ChangedDataTables, ChangedCurveTablesPointer);
+									bAddAssetToHotfixedList = ProblemStrings.Num() == 0;
+								}
+								else
+								{
+									ProblemStrings.Add(FString::Printf(TEXT("Expected a hotfix type RowUpdate with 5 tokens but got %d"), Tokens.Num()));
+								}
 							}
-							else if ((HotfixType == TableUpdate || HotfixType == CurveUpdate) && Tokens.Num() == 3)
+							else if (HotfixType == AddRow)
 							{
-								// The hotfix line should be
-								//	+DataTable=<data table path>;TableUpdate;"<json data>"
-								//	+CurveTable=<curve table path>;TableUpdate;"<json data>"
-								//	+CurveFloat=<curve float path>;CurveUpdate;"<json data>"
-								//	+CurveVector=<curve vector path>;CurveUpdate;"<json data>"
-								//	+CurveLinearColor=<curve linear color path>;CurveUpdate;"<json data>"
-
-								// We have to read json data as quoted string because tokenizing it creates extra unwanted characters.
-								FString JsonData;
-
-								// Json should be read in its entirety, if the whole buffer wasn't read the string is malformed. 
-								int32 ReadLen = 0;
-								int32 InputLen = Tokens[2].Len();
-								if (FParse::QuotedString(*Tokens[2], JsonData, &ReadLen) && ReadLen == InputLen)
+								if (Tokens.Num() == 3)
 								{
-									HotfixTableUpdate(Asset, AssetPath, JsonData, ProblemStrings);
-									bAddAssetToHotfixedList = ProblemStrings.Num() == 0;
+									// The hotfix line should be
+									//	+DataTable=<data table path>;AddRow;"<json data>"
+
+									// We have to read json data as quoted string because tokenizing it creates extra unwanted characters.
+									FString JsonData;
+
+									// Json should be read in its entirety, if the whole buffer wasn't read the string is malformed. 
+									int32 ReadLen = 0;
+									int32 InputLen = Tokens[2].Len();
+									if (FParse::QuotedString(*Tokens[2], JsonData, &ReadLen) && ReadLen == InputLen)
+									{
+										HotfixAddRow(Asset, AssetPath, JsonData, ProblemStrings);
+										bAddAssetToHotfixedList = ProblemStrings.Num() == 0;
+									}
+									else
+									{
+										ProblemStrings.Add(TEXT("Json data wasn't able to be parsed as a quoted string. Check that we have opening and closing quotes around the json data."));
+									}
 								}
 								else
 								{
-									ProblemStrings.Add(TEXT("Json data wasn't able to be parsed as a quoted string. Check that we have opening and closing quotes around the json data."));
+									ProblemStrings.Add(FString::Printf(TEXT("Expected a hotfix type AddRow with 3 tokens but got %d"), Tokens.Num()));
+								}
+							}
+							else if (HotfixType == TableUpdate || HotfixType == CurveUpdate)
+							{
+								if (Tokens.Num() == 3)
+								{
+									// The hotfix line should be
+									//	+DataTable=<data table path>;TableUpdate;"<json data>"
+									//	+CurveTable=<curve table path>;TableUpdate;"<json data>"
+									//	+CurveFloat=<curve float path>;CurveUpdate;"<json data>"
+									//	+CurveVector=<curve vector path>;CurveUpdate;"<json data>"
+									//	+CurveLinearColor=<curve linear color path>;CurveUpdate;"<json data>"
+
+									// We have to read json data as quoted string because tokenizing it creates extra unwanted characters.
+									FString JsonData;
+
+									// Json should be read in its entirety, if the whole buffer wasn't read the string is malformed. 
+									int32 ReadLen = 0;
+									int32 InputLen = Tokens[2].Len();
+									if (FParse::QuotedString(*Tokens[2], JsonData, &ReadLen) && ReadLen == InputLen)
+									{
+										HotfixTableUpdate(Asset, AssetPath, JsonData, ProblemStrings);
+										bAddAssetToHotfixedList = ProblemStrings.Num() == 0;
+									}
+									else
+									{
+										ProblemStrings.Add(TEXT("Json data wasn't able to be parsed as a quoted string. Check that we have opening and closing quotes around the json data."));
+									}
+								}
+								else
+								{
+									ProblemStrings.Add(FString::Printf(TEXT("Expected a hotfix type %s with 3 tokens but got %d"), *HotfixType, Tokens.Num()));
 								}
 							}
 							else
 							{
-								ProblemStrings.Add(TEXT("Expected a hotfix type of RowUpdate with 5 tokens or TableUpdate/CurveUpdate with 3 tokens."));
+								ProblemStrings.Add(FString::Printf(TEXT("Unknown hotfix type %s"), *HotfixType));
 							}
 						}
 						else
@@ -1918,6 +1961,47 @@ void UOnlineHotfixManager::HotfixRowUpdate(
 	}
 }
 
+void UOnlineHotfixManager::HotfixAddRow(
+	UObject* Asset,
+	const FString& AssetPath,
+	const FString& JsonData,
+	TArray<FString>& ProblemStrings,
+	TSet<UDataTable*>* ChangedDataTables)
+{
+	if (AssetPath.IsEmpty())
+	{
+		ProblemStrings.Add(TEXT("The table's path is empty. We cannot continue the hotfix."));
+		return;
+	}
+	if (JsonData.IsEmpty())
+	{
+		ProblemStrings.Add(TEXT("The JSON data is empty. We cannot continue the hotfix."));
+		return;
+	}
+
+	if (UDataTable* DataTable = Cast<UDataTable>(Asset))
+	{
+		FName RowName = NAME_None;
+		// Allow hotfix to be applied multiple times
+		const bool bRemoveDuplicate = true;
+		if (!DataTableUtils::AddRowJSON(*DataTable, JsonData, ProblemStrings, bRemoveDuplicate, &RowName))
+		{
+			return;
+		}
+
+		OnHotfixTableAddRow(*DataTable, RowName);
+
+		if (ChangedDataTables == nullptr)
+		{
+			DataTable->HandleDataTableChanged();
+		}
+		else
+		{
+			ChangedDataTables->Add(DataTable);
+		}
+	}
+}
+
 void UOnlineHotfixManager::HotfixTableUpdate(UObject* Asset, const FString& AssetPath, const FString& JsonData, TArray<FString>& ProblemStrings)
 {
 	if (AssetPath.IsEmpty())
diff --git a/Engine/Plugins/Online/OnlineFramework/Source/Hotfix/Public/OnlineHotfixManager.h b/Engine/Plugins/Online/OnlineFramework/Source/Hotfix/Public/OnlineHotfixManager.h
--- a/Engine/Plugins/Online/OnlineFramework/Source/Hotfix/Public/OnlineHotfixManager.h
+++ b/Engine/Plugins/Online/OnlineFramework/Source/Hotfix/Public/OnlineHotfixManager.h
@@ -359,7 +359,20 @@ protected:
 		TArray<FString>& ProblemStrings,
 		TSet<class UDataTable*>* ChangedDataTables = nullptr,
 		TSet<class UCurveTable*>* ChangedCurveTables = nullptr);
-	
+
+	/** Used in PatchAssetsFromIniFiles to hotfix a new row in a table.
+	 *  If ChangedTables is not null then HandleDataTableChanged will not be called and the caller should call it on the data tables in ChangedTables when they're ready to
+	 */
+	void HotfixAddRow(
+		UObject* Asset,
+		const FString& AssetPath,
+		const FString& JsonData,
+		TArray<FString>& ProblemStrings,
+		TSet<class UDataTable*>* ChangedDataTables = nullptr);
+
+	/** Called after adding table row by HotfixAddRow() */
+	virtual void OnHotfixTableAddRow(UObject& Asset, const FName RowName) {}
+
 	/** Used in PatchAssetsFromIniFiles to hotfix an entire table. */
 	void HotfixTableUpdate(UObject* Asset, const FString& AssetPath, const FString& JsonData, TArray<FString>& ProblemStrings);
 
diff --git a/Engine/Source/Runtime/Engine/Private/DataTableJSON.cpp b/Engine/Source/Runtime/Engine/Private/DataTableJSON.cpp
--- a/Engine/Source/Runtime/Engine/Private/DataTableJSON.cpp
+++ b/Engine/Source/Runtime/Engine/Private/DataTableJSON.cpp
@@ -446,7 +446,42 @@ bool FDataTableImporterJSON::ReadTable()
 	return true;
 }
 
-bool FDataTableImporterJSON::ReadRow(const TSharedRef<FJsonObject>& InParsedTableRowObject, const int32 InRowIdx)
+bool FDataTableImporterJSON::AddRow(bool bRemoveDuplicate, FName* OutRowName)
+{
+	if (JSONData.IsEmpty())
+	{
+		ImportProblems.Add(TEXT("Input data is empty."));
+		return false;
+	}
+
+	// Check we have a RowStruct specified
+	if (!DataTable->RowStruct)
+	{
+		ImportProblems.Add(TEXT("No RowStruct specified."));
+		return false;
+	}
+
+	TSharedPtr<FJsonValue> ParsedTableRow;
+	{
+		const TSharedRef< TJsonReader<TCHAR> > JsonReader = TJsonReaderFactory<TCHAR>::Create(JSONData);
+		if (!FJsonSerializer::Deserialize(JsonReader, ParsedTableRow) || !ParsedTableRow.IsValid())
+		{
+			ImportProblems.Add(FString::Printf(TEXT("Failed to parse the JSON data. Error: %s"), *JsonReader->GetErrorMessage()));
+			return false;
+		}
+	}
+
+	TSharedPtr<FJsonObject> ParsedTableRowObject = ParsedTableRow->AsObject();
+	if (!ParsedTableRowObject.IsValid())
+	{
+		ImportProblems.Add(TEXT("Row is not a valid JSON object."));
+		return false;
+	}
+
+	return ReadRow(ParsedTableRowObject.ToSharedRef(), 0, bRemoveDuplicate, OutRowName);
+}
+
+bool FDataTableImporterJSON::ReadRow(const TSharedRef<FJsonObject>& InParsedTableRowObject, const int32 InRowIdx, bool bRemoveDuplicate, FName* OutRowName)
 {
 	// Get row name
 	FString RowKey = DataTableJSONUtils::GetKeyFieldName(*DataTable);
@@ -458,12 +493,23 @@ bool FDataTableImporterJSON::ReadRow(const TSharedRef<FJsonObject>& InParsedTabl
 		ImportProblems.Add(FString::Printf(TEXT("Row '%d' missing key field '%s'."), InRowIdx, *RowKey));
 		return false;
 	}
+	else if (OutRowName != nullptr)
+	{
+		*OutRowName = RowName;
+	}
 
 	// Check its not a duplicate
-	if (!DataTable->AllowDuplicateRowsOnImport() && DataTable->GetRowMap().Find(RowName) != nullptr)
+	if (DataTable->GetRowMap().Find(RowName) != nullptr)
 	{
-		ImportProblems.Add(FString::Printf(TEXT("Duplicate row name '%s'."), *RowName.ToString()));
-		return false;
+		if (bRemoveDuplicate)
+		{
+			DataTable->RemoveRowInternal(RowName);
+		}
+		else if (!DataTable->AllowDuplicateRowsOnImport())
+		{
+			ImportProblems.Add(FString::Printf(TEXT("Duplicate row name '%s'."), *RowName.ToString()));
+			return false;
+		}
 	}
 
 	// Detect any extra fields within the data for this row
diff --git a/Engine/Source/Runtime/Engine/Private/DataTableJSON.h b/Engine/Source/Runtime/Engine/Private/DataTableJSON.h
--- a/Engine/Source/Runtime/Engine/Private/DataTableJSON.h
+++ b/Engine/Source/Runtime/Engine/Private/DataTableJSON.h
@@ -82,9 +82,10 @@ public:
 	~FDataTableImporterJSON();
 
 	bool ReadTable();
+	bool AddRow(bool bRemoveDuplicate = false, FName* OutRowName = nullptr);
 
 private:
-	bool ReadRow(const TSharedRef<FJsonObject>& InParsedTableRowObject, const int32 InRowIdx);
+	bool ReadRow(const TSharedRef<FJsonObject>& InParsedTableRowObject, const int32 InRowIdx, bool bRemoveDuplicate = false, FName* OutRowName = nullptr);
 
 	bool ReadStruct(const TSharedRef<FJsonObject>& InParsedObject, UScriptStruct* InStruct, const FName InRowName, void* InStructData);
 
diff --git a/Engine/Source/Runtime/Engine/Private/DataTableUtils.cpp b/Engine/Source/Runtime/Engine/Private/DataTableUtils.cpp
--- a/Engine/Source/Runtime/Engine/Private/DataTableUtils.cpp
+++ b/Engine/Source/Runtime/Engine/Private/DataTableUtils.cpp
@@ -488,3 +488,8 @@ TArray<FString> DataTableUtils::GetColumnDataAsString(const UDataTable* InTable,
 
 	return Result;
 }
+
+bool DataTableUtils::AddRowJSON(UDataTable& InDataTable, const FString& InJSONData, TArray<FString>& OutProblems, bool bRemoveDuplicate, FName* OutRowName)
+{
+	return FDataTableImporterJSON(InDataTable, InJSONData, OutProblems).AddRow(bRemoveDuplicate, OutRowName);
+}
\ No newline at end of file
diff --git a/Engine/Source/Runtime/Engine/Public/DataTableUtils.h b/Engine/Source/Runtime/Engine/Public/DataTableUtils.h
--- a/Engine/Source/Runtime/Engine/Public/DataTableUtils.h
+++ b/Engine/Source/Runtime/Engine/Public/DataTableUtils.h
@@ -97,4 +97,9 @@ namespace DataTableUtils
 	 * Output each row for a specific column/property in the table (doesn't include the title)
 	 */
 	ENGINE_API TArray<FString> GetColumnDataAsString(const UDataTable* InTable, const FName& PropertyName, const EDataTableExportFlags InDTExportFlags);
+
+	/**
+	 * Add a row from a json format to a datatable
+	 */
+	ENGINE_API bool AddRowJSON(UDataTable& InDataTable, const FString& InJSONData, TArray<FString>& OutProblems, bool bRemoveDuplicate = false, FName* OutRowName = nullptr);
 }
