Add support to OnlineHotfixManager to fallback to hotfix manager type specified in DefaultEngine.ini if the most specific requested type is unavailable.

diff --git a/Engine/Plugins/Online/OnlineFramework/Source/Hotfix/Private/OnlineHotfixManager.cpp b/Engine/Plugins/Online/OnlineFramework/Source/Hotfix/Private/OnlineHotfixManager.cpp
--- a/Engine/Plugins/Online/OnlineFramework/Source/Hotfix/Private/OnlineHotfixManager.cpp
+++ b/Engine/Plugins/Online/OnlineFramework/Source/Hotfix/Private/OnlineHotfixManager.cpp
@@ -222,8 +222,22 @@ UOnlineHotfixManager* UOnlineHotfixManager::Get(UWorld* World)
 			UClass* HotfixManagerClass = LoadClass<UOnlineHotfixManager>(nullptr, *HotfixManagerClassName, nullptr, LOAD_None, nullptr);
 			if (HotfixManagerClass == nullptr)
 			{
-				// Just use the default class if it couldn't load what was specified
-				HotfixManagerClass = UOnlineHotfixManager::StaticClass();
+				// Retrieve the static value, before any plugins, hotfixes, or any other dynamic configs have been applied, as a fallback
+				// This is not a standard method for retrieving config values, and should not be used as a general solution in other code
+				FString FallbackClassName;
+				GConfig->FindBranch(TEXT("Engine"), TEXT("DefaultEngine"))->CombinedStaticLayers.GetString(TEXT("/Script/Hotfix.OnlineHotfixManager"), TEXT("HotfixManagerClassName"), FallbackClassName);
+				if (FallbackClassName.Len())
+				{
+					UE_LOG(LogHotfixManager, Display, TEXT("Unable to locate hotfix manager class '%s'. Trying '%s' instead."), *HotfixManagerClassName, *FallbackClassName)
+					HotfixManagerClassName = FallbackClassName;
+					HotfixManagerClass = LoadClass<UOnlineHotfixManager>(nullptr, *HotfixManagerClassName, nullptr, LOAD_None, nullptr);
+				}
+				if (HotfixManagerClass == nullptr)
+				{
+					UE_LOG(LogHotfixManager, Warning, TEXT("Failed to find requested hotfix manager class '%s'. Falling back to default."), *HotfixManagerClassName);
+					// Just use the default class if it couldn't load what was specified
+					HotfixManagerClass = UOnlineHotfixManager::StaticClass();
+				}
 			}
 			// Create it and store it
 			HotfixManager = NewObject<UOnlineHotfixManager>(GetTransientPackage(), HotfixManagerClass);
