[AutoRTFM] Defer hotfix loading when inside a transaction.

Loading the hotfix list triggers an HTTP request, which is fundamentally not something we can roll back. Fortunately, it's also an asynchronous process, so this deferral should be invisible to the caller.

Additionally, I moved a few extra statements into the `IsValid` check because they access the pointer being checked for validity, and would crash if it were null anyway. Presumably this pointer has always been valid during our testing (or is valid often enough that crashes are below the noise floor).

diff --git a/Engine/Plugins/Online/OnlineFramework/Source/Hotfix/Private/OnlineHotfixManager.cpp b/Engine/Plugins/Online/OnlineFramework/Source/Hotfix/Private/OnlineHotfixManager.cpp
--- a/Engine/Plugins/Online/OnlineFramework/Source/Hotfix/Private/OnlineHotfixManager.cpp
+++ b/Engine/Plugins/Online/OnlineFramework/Source/Hotfix/Private/OnlineHotfixManager.cpp
@@ -24,6 +24,7 @@
 #include "Curves/CurveLinearColor.h"
 #include "Engine/BlueprintGeneratedClass.h"
 
+#include "AutoRTFM/AutoRTFM.h"
 #include "Serialization/AsyncLoadingFlushContext.h"
 #include "UObject/PropertyAccessUtil.h"
 
@@ -525,16 +526,20 @@ void UOnlineHotfixManager::CheckAvailability(FOnHotfixAvailableComplete& InCompl
 
 	OnlineTitleFile = Online::GetTitleFileInterface(OSSName.Len() ? FName(*OSSName, FNAME_Find) : NAME_None);
 
-	FOnEnumerateFilesCompleteDelegate OnEnumerateFilesForAvailabilityCompleteDelegate;
-	OnEnumerateFilesForAvailabilityCompleteDelegate.BindUObject(this, &UOnlineHotfixManager::OnEnumerateFilesForAvailabilityComplete, InCompletionDelegate);
-	OnEnumerateFilesForAvailabilityCompleteDelegateHandle = OnlineTitleFile->AddOnEnumerateFilesCompleteDelegate_Handle(OnEnumerateFilesForAvailabilityCompleteDelegate);
-
 	bHotfixingInProgress = true;
 
 	// Kick off an enumeration of the files that are available to download
 	if (OnlineTitleFile.IsValid())
 	{
-		OnlineTitleFile->EnumerateFiles();
+		FOnEnumerateFilesCompleteDelegate OnEnumerateFilesForAvailabilityCompleteDelegate;
+		OnEnumerateFilesForAvailabilityCompleteDelegate.BindUObject(this, &UOnlineHotfixManager::OnEnumerateFilesForAvailabilityComplete, InCompletionDelegate);
+		OnEnumerateFilesForAvailabilityCompleteDelegateHandle = OnlineTitleFile->AddOnEnumerateFilesCompleteDelegate_Handle(OnEnumerateFilesForAvailabilityCompleteDelegate);
+
+		// The hotfix process isn't purely reversible as it kicks off HTTP requests, so we defer the work until transaction commit time.
+		UE_AUTORTFM_ONCOMMIT(this)
+		{
+			OnlineTitleFile->EnumerateFiles();
+		};
 	}
 	else
 	{
