Only check HF availability and do not apply them in UpdateManager::OnApplicationHasReactivated

#rnx

diff --git a/Engine/Plugins/Online/OnlineFramework/Source/Hotfix/Private/UpdateManager.cpp b/Engine/Plugins/Online/OnlineFramework/Source/Hotfix/Private/UpdateManager.cpp
--- a/Engine/Plugins/Online/OnlineFramework/Source/Hotfix/Private/UpdateManager.cpp
+++ b/Engine/Plugins/Online/OnlineFramework/Source/Hotfix/Private/UpdateManager.cpp
@@ -629,7 +629,7 @@ void UUpdateManager::OnApplicationHasReactivated()
 
 	if ((Now - DeactivatedTime).GetTotalSeconds() > AppSuspendedUpdateCheckTimeSeconds)
 	{
-		StartCheck();
+		StartUpdateCheck(TEXT("ApplicationReactivated"));
 	}
 }
 
