Add LOGHOTFIXASSETLIST command for logging the contents of UOnlineHotfixManager::AssetsHotfixedFromIniFiles

diff --git a/Engine/Plugins/Online/OnlineFramework/Source/Hotfix/Private/OnlineHotfixManager.cpp b/Engine/Plugins/Online/OnlineFramework/Source/Hotfix/Private/OnlineHotfixManager.cpp
--- a/Engine/Plugins/Online/OnlineFramework/Source/Hotfix/Private/OnlineHotfixManager.cpp
+++ b/Engine/Plugins/Online/OnlineFramework/Source/Hotfix/Private/OnlineHotfixManager.cpp
@@ -2100,6 +2100,8 @@ void UOnlineHotfixManager::HotfixDynamicBranch(const FName& Tag, const FName& Br
 	const FString HotfixBaseName = Tag.ToString() + Branch.ToString();
 	const TArray<TPair<FString, FString>>* HotfixesForTag = DynamicHotfixContents.Find(*HotfixBaseName);
 
+	UE_LOG(LogHotfixManager, VeryVerbose, TEXT("HotfixDynamicBranch: Searching for hotfix for base name: %s. Result: %s"), *HotfixBaseName, HotfixesForTag ? TEXT("found") : TEXT("not found"));
+
 	if (HotfixesForTag)
 	{
 		for (const TPair<FString, FString>& HotfixFilenameAndContents : *HotfixesForTag)
@@ -2275,6 +2277,7 @@ void UOnlineHotfixManager::ReloadObjectsAffectedByConfigFile(const FString& IniD
 		*IniDataFileName, FPlatformTime::Seconds() - StartTime, NumObjectsReloaded);
 }
 
+#if UE_ALLOW_EXEC_COMMANDS
 struct FHotfixManagerExec :
 	public FSelfRegisteringExec
 {
@@ -2348,8 +2351,24 @@ protected:
 			}
 			return true;
 		}
+		else if (FParse::Command(&Cmd, TEXT("LOGHOTFIXASSETLIST")))
+		{
+			if (UOnlineHotfixManager* HotfixManager = UOnlineHotfixManager::Get(InWorld))
+			{
+				int32 CurrentAsset = 0;
+				const int32 AssetCount = HotfixManager->AssetsHotfixedFromIniFiles.Num();
+				UE_LOG(LogHotfixManager, Log, TEXT("Logging contetents of AssetsHotfixedFromIniFiles. Asset count: %d"), AssetCount);
+				for (const TObjectPtr<UObject>& ObjPtr : HotfixManager->AssetsHotfixedFromIniFiles)
+				{
+					++CurrentAsset;
+					UE_LOG(LogHotfixManager, Log, TEXT("\t[%d/%d]: %s"), CurrentAsset, AssetCount, *GetPathNameSafe(ObjPtr.Get()));
+				}
+				UE_LOG(LogHotfixManager, Log, TEXT("Logging complete."));
+			}
+			return true;
+		}
 		return false;
 	}
 };
 static FHotfixManagerExec HotfixManagerExec;
-
+#endif // UE_ALLOW_EXEC_COMMANDS
