PR #12306: Fix Crash in Lobby Beacon Host for Logging

diff --git a/Engine/Plugins/Online/OnlineFramework/Source/Lobby/Private/LobbyBeaconHost.cpp b/Engine/Plugins/Online/OnlineFramework/Source/Lobby/Private/LobbyBeaconHost.cpp
--- a/Engine/Plugins/Online/OnlineFramework/Source/Lobby/Private/LobbyBeaconHost.cpp
+++ b/Engine/Plugins/Online/OnlineFramework/Source/Lobby/Private/LobbyBeaconHost.cpp
@@ -133,7 +133,8 @@ void ALobbyBeaconHost::ProcessLogin(ALobbyBeaconClient* ClientActor, const FStri
 		*InSessionId,
 		*InUniqueId.ToString(),
 		*UrlString,
-		ClientActor ? *ClientActor->GetNetConnection()->LowLevelDescribe() : TEXT("NULL"));
+		ClientActor ? (ClientActor->GetNetConnection() ? *ClientActor->GetNetConnection()->LowLevelDescribe() : TEXT("NULL")) : TEXT("NULL"));
+
 
 	if (ClientActor)
 	{
