Deprecated unused method in SocialToolkit

diff --git a/Engine/Plugins/Online/OnlineFramework/Source/Party/Public/SocialToolkit.h b/Engine/Plugins/Online/OnlineFramework/Source/Party/Public/SocialToolkit.h
--- a/Engine/Plugins/Online/OnlineFramework/Source/Party/Public/SocialToolkit.h
+++ b/Engine/Plugins/Online/OnlineFramework/Source/Party/Public/SocialToolkit.h
@@ -76,6 +76,8 @@ public:
 
 	const FOnlineUserPresence* GetPresenceInfo(ESocialSubsystem SubsystemType) const;
 	void SetLocalUserOnlineState(EOnlinePresenceState::Type OnlineState);
+
+	UE_DEPRECATED(5.6, "AddLocalUserOnlineProperties is deprecated, please use IOnlinePresence::SetPresence instead.")
 	void AddLocalUserOnlineProperties(FPresenceProperties OnlineProperties);
 
 	UE_DEPRECATED(5.5, "GetOwningLocalPlayer returns LocalPlayerOwner which is a TWeakObjectPtr and may return nullptr when the local player logs out. Please use the pointer version.")
