Don't check friend's platform session when check if the user can Request To Join.

diff --git a/Engine/Plugins/Online/OnlineFramework/Source/Party/Private/SocialManager.cpp b/Engine/Plugins/Online/OnlineFramework/Source/Party/Private/SocialManager.cpp
--- a/Engine/Plugins/Online/OnlineFramework/Source/Party/Private/SocialManager.cpp
+++ b/Engine/Plugins/Online/OnlineFramework/Source/Party/Private/SocialManager.cpp
@@ -493,6 +493,11 @@ FJoinPartyResult USocialManager::ValidateJoinAttempt(const FOnlinePartyTypeId& P
 }
 
 FJoinPartyResult USocialManager::ValidateJoinTarget(const USocialUser& UserToJoin, const FOnlinePartyTypeId& PartyTypeId) const
+{
+	return ValidateJoinTarget(UserToJoin, PartyTypeId, /*bCheckPlatformSession*/true);
+}
+
+FJoinPartyResult USocialManager::ValidateJoinTarget(const USocialUser& UserToJoin, const FOnlinePartyTypeId& PartyTypeId, bool bCheckPlatformSession) const
 {
 	UE_LOG(LogParty, VeryVerbose, TEXT("Validating user [%s] as join target of party type [%d]"), *UserToJoin.ToDebugString(), PartyTypeId.GetValue());
 
@@ -557,7 +562,14 @@ FJoinPartyResult USocialManager::ValidateJoinTarget(const USocialUser& UserToJoi
 					}
 					else if (!PlatformPresence->SessionId.IsValid() || !PlatformPresence->SessionId->IsValid())
 					{
-						return FPartyJoinDenialReason(EPartyJoinDenialReason::TargetUserMissingPlatformSession);
+						if (bCheckPlatformSession)
+						{
+							return FPartyJoinDenialReason(EPartyJoinDenialReason::TargetUserMissingPlatformSession);
+						}
+						else
+						{
+							return FPartyJoinDenialReason(EPartyJoinDenialReason::PartyPrivate);
+						}
 					}
 				}
 			}
@@ -586,7 +598,7 @@ void USocialManager::JoinParty(const USocialUser& UserToJoin, const FOnlineParty
 
 	FJoinPartyAttempt NewAttempt(&UserToJoin, PartyTypeId, JoinMethod, OnJoinPartyComplete);
 	NewAttempt.AnalyticsContext = UserToJoin.GetAnalyticsContext();
-	const FJoinPartyResult ValidationResult = ValidateJoinTarget(UserToJoin, PartyTypeId);
+	const FJoinPartyResult ValidationResult = ValidateJoinTarget(UserToJoin, PartyTypeId, /*bCheckPlatformSession*/true);
 	if (ValidationResult.WasSuccessful())
 	{
 		FJoinPartyAttempt& JoinAttempt = JoinAttemptsByTypeId.Add(PartyTypeId, NewAttempt);
diff --git a/Engine/Plugins/Online/OnlineFramework/Source/Party/Private/User/SocialUser.cpp b/Engine/Plugins/Online/OnlineFramework/Source/Party/Private/User/SocialUser.cpp
--- a/Engine/Plugins/Online/OnlineFramework/Source/Party/Private/User/SocialUser.cpp
+++ b/Engine/Plugins/Online/OnlineFramework/Source/Party/Private/User/SocialUser.cpp
@@ -1152,9 +1152,9 @@ bool USocialUser::EndFriendship(ESocialSubsystem SocialSubsystem) const
 	return false;
 }
 
-FJoinPartyResult USocialUser::CheckPartyJoinability(const FOnlinePartyTypeId& PartyTypeId) const
+FJoinPartyResult USocialUser::CheckPartyJoinability(const FOnlinePartyTypeId& PartyTypeId, bool bCheckPlatformSession/*= true*/) const
 {
-	return GetOwningToolkit().GetSocialManager().ValidateJoinTarget(*this, PartyTypeId);
+	return GetOwningToolkit().GetSocialManager().ValidateJoinTarget(*this, PartyTypeId, bCheckPlatformSession);
 }
 
 bool USocialUser::ShowPlatformProfile()
diff --git a/Engine/Plugins/Online/OnlineFramework/Source/Party/Public/SocialManager.h b/Engine/Plugins/Online/OnlineFramework/Source/Party/Public/SocialManager.h
--- a/Engine/Plugins/Online/OnlineFramework/Source/Party/Public/SocialManager.h
+++ b/Engine/Plugins/Online/OnlineFramework/Source/Party/Public/SocialManager.h
@@ -118,8 +118,11 @@ public:
 	virtual void NotifyPartyInitialized(USocialParty& Party);
 
 	/** Validates that the target user has valid join info for us to use and that we can join any party of the given type */
+	UE_DEPRECATED(5.6, "ValidateJoinTarget has been deprecated, use the overloaded one instead")
 	virtual FJoinPartyResult ValidateJoinTarget(const USocialUser& UserToJoin, const FOnlinePartyTypeId& PartyTypeId) const;
 
+	virtual FJoinPartyResult ValidateJoinTarget(const USocialUser& UserToJoin, const FOnlinePartyTypeId& PartyTypeId, bool bCheckPlatformSession) const;
+
 protected:
 	DECLARE_DELEGATE_OneParam(FOnJoinPartyAttemptComplete, const FJoinPartyResult&);
 	void JoinParty(const USocialUser& UserToJoin, const FOnlinePartyTypeId& PartyTypeId, const FOnJoinPartyAttemptComplete& OnJoinPartyComplete, const FName& JoinMethod);
diff --git a/Engine/Plugins/Online/OnlineFramework/Source/Party/Public/User/SocialUser.h b/Engine/Plugins/Online/OnlineFramework/Source/Party/Public/User/SocialUser.h
--- a/Engine/Plugins/Online/OnlineFramework/Source/Party/Public/User/SocialUser.h
+++ b/Engine/Plugins/Online/OnlineFramework/Source/Party/Public/User/SocialUser.h
@@ -130,7 +130,7 @@ public:
 	virtual IOnlinePartyJoinInfoConstPtr GetPartyJoinInfo(const FOnlinePartyTypeId& PartyTypeId) const;
 
 	bool HasSentPartyInvite(const FOnlinePartyTypeId& PartyTypeId) const;
-	FJoinPartyResult CheckPartyJoinability(const FOnlinePartyTypeId& PartyTypeId) const;
+	FJoinPartyResult CheckPartyJoinability(const FOnlinePartyTypeId& PartyTypeId, bool bCheckPlatformSession = true) const;
 
 	virtual void JoinParty(const FOnlinePartyTypeId& PartyTypeId, const FName& JoinMethod) const;
 	virtual void RejectPartyInvite(const FOnlinePartyTypeId& PartyTypeId);
