Allow canceling party JIP requests.

diff --git a/Engine/Plugins/Online/OnlineFramework/Source/Party/Private/Party/SocialParty.cpp b/Engine/Plugins/Online/OnlineFramework/Source/Party/Private/Party/SocialParty.cpp
--- a/Engine/Plugins/Online/OnlineFramework/Source/Party/Private/Party/SocialParty.cpp
+++ b/Engine/Plugins/Online/OnlineFramework/Source/Party/Private/Party/SocialParty.cpp
@@ -1942,14 +1942,10 @@ void USocialParty::JoinSessionCompleteAnalytics(const FSessionId& SessionId, con
 
 void USocialParty::RequestJoinInProgress(const UPartyMember& TargetMember, const FOnRequestJoinInProgressComplete& CompletionDelegate)
 {
-	// Only allow one join attempt at a time
-	if (RequestJoinInProgressComplete.IsSet())
+	// When a join-in-progress attempt is already in progress, cancel the previous attempt before creating a new one.
+	if (IsJoinInProgeressRequestActive())
 	{
-		UE_LOG(LogParty, Warning, TEXT("RequestJoinInProgress: Request already in progress"));
-		GetWorld()->GetTimerManager().SetTimerForNextTick(FTimerDelegate::CreateLambda([CompletionDelegate]() {
-			CompletionDelegate.ExecuteIfBound(EPartyJoinDenialReason::Busy);
-		}));
-		return;
+		CancelJoinInProgressRequest();
 	}
 
 	RequestJoinInProgressComplete = CompletionDelegate;
@@ -1968,6 +1964,30 @@ void USocialParty::RequestJoinInProgress(const UPartyMember& TargetMember, const
 	RunJoinInProgressTimer();
 }
 
+void USocialParty::CancelJoinInProgressRequest()
+{
+	if (IsJoinInProgeressRequestActive())
+	{
+		UE_LOG(LogParty, Log, TEXT("RequestJoinInProgress: Canceling join-in-progress request."));
+
+		FPartyMemberJoinInProgressRequest EmptyRequest;
+		for (UPartyMember* RequestingMember : GetLocalPartyMembersForJoinInProgress())
+		{
+			if (RequestingMember && ensure(RequestingMember->IsLocalPlayer()))
+			{
+				RequestingMember->GetMutableRepData().SetJoinInProgressDataRequest(EmptyRequest);
+			}
+		}
+
+		CallJoinInProgressComplete(EPartyJoinDenialReason::JoinAttemptAborted);
+	}
+}
+
+bool USocialParty::IsJoinInProgeressRequestActive() const
+{
+	return RequestJoinInProgressComplete.IsSet();
+}
+
 void USocialParty::HandleJoinInProgressDataRequestChanged(const FPartyMemberJoinInProgressRequest& Request, UPartyMember* Member)
 {
 	if (Request.Time == 0 || !(AllowJoinInProgressToMember() || IsLocalPlayerPartyLeader()))
diff --git a/Engine/Plugins/Online/OnlineFramework/Source/Party/Public/Party/SocialParty.h b/Engine/Plugins/Online/OnlineFramework/Source/Party/Public/Party/SocialParty.h
--- a/Engine/Plugins/Online/OnlineFramework/Source/Party/Public/Party/SocialParty.h
+++ b/Engine/Plugins/Online/OnlineFramework/Source/Party/Public/Party/SocialParty.h
@@ -222,6 +222,8 @@ public:
 
 	DECLARE_DELEGATE_OneParam(FOnRequestJoinInProgressComplete, const EPartyJoinDenialReason /*DenialReason*/);
 	UE_API void RequestJoinInProgress(const UPartyMember& TargetMember, const FOnRequestJoinInProgressComplete& CompletionDelegate);
+	UE_API void CancelJoinInProgressRequest();
+	UE_API bool IsJoinInProgeressRequestActive() const;
 
 	UE_API bool CanInviteUser(const USocialUser& User, const ESocialPartyInviteMethod InviteMethod = ESocialPartyInviteMethod::Other) const;
 
