SWT/SWT2 separation in social code

diff --git a/Engine/Config/BaseGame.ini b/Engine/Config/BaseGame.ini
--- a/Engine/Config/BaseGame.ini
+++ b/Engine/Config/BaseGame.ini
@@ -57,7 +57,8 @@ DefaultMaxPartySize=5
 +SocialPlatformDescriptions=(Name="MAC", PlatformType="DESKTOP", CrossplayPool="DESKTOP")
 +SocialPlatformDescriptions=(Name="WIN", PlatformType="DESKTOP", CrossplayPool="DESKTOP")
 +SocialPlatformDescriptions=(Name="PSN", PlatformType="CONSOLE", OnlineSubsystem="PS4", ExternalAccountType="psn", SessionType="PS4", CrossplayPool="psn")
-+SocialPlatformDescriptions=(Name="SWT", PlatformType="CONSOLE", OnlineSubsystem="SWITCH", ExternalAccountType="nintendo", CrossplayPool="nintendo")
++SocialPlatformDescriptions=(Name="SWT", PlatformType="CONSOLE", OnlineSubsystem="NINTENDO", ExternalAccountType="nintendo", CrossplayPool="nintendo")
++SocialPlatformDescriptions=(Name="SWT2", PlatformType="CONSOLE", OnlineSubsystem="NINTENDO", ExternalAccountType="nintendo", CrossplayPool="nintendo")
 +SocialPlatformDescriptions=(Name="XBL", PlatformType="CONSOLE", OnlineSubsystem="LIVE", ExternalAccountType="xbl", SessionType="MPSD", CrossplayPool="xbl")
 
 [/Script/Lobby.LobbyBeaconState]
diff --git a/Engine/Plugins/Online/OnlineFramework/Source/Party/Private/User/SocialUser.cpp b/Engine/Plugins/Online/OnlineFramework/Source/Party/Private/User/SocialUser.cpp
--- a/Engine/Plugins/Online/OnlineFramework/Source/Party/Private/User/SocialUser.cpp
+++ b/Engine/Plugins/Online/OnlineFramework/Source/Party/Private/User/SocialUser.cpp
@@ -729,6 +729,9 @@ FUserPlatform USocialUser::GetCurrentPlatform() const
 		const FOnlineUserPresence* PrimaryPresence = GetFriendPresenceInfo(ESocialSubsystem::Primary);
 		const FOnlineUserPresence* PlatformPresence = GetFriendPresenceInfo(ESocialSubsystem::Platform);
 
+// jniman: Switch and Switch2 must be differentiated from one another, but they have the same OSS, so some of these checks will return bad results and cannot be ran on these platforms
+// jntodo: this will need to get moved to a platform-specific uplugin once we have the time
+#if !((defined(PLATFORM_SWITCH) && PLATFORM_SWITCH) || (defined(PLATFORM_SWITCH2) && PLATFORM_SWITCH2))
 		if (PlatformPresence && PlatformPresence->bIsOnline && PlatformPresence->bIsPlayingThisGame)
 		{
 			// Platform friends that are playing the same game are on the local platform
@@ -744,7 +747,9 @@ FUserPlatform USocialUser::GetCurrentPlatform() const
 			// Not playing the same game on either account, but we have presence on the platform, so let that win regardless of whether the primary is valid
 			PlatformString = IOnlineSubsystem::GetLocalPlatformName();
 		}
-		else if (PrimaryPresence && PrimaryPresence->bIsOnline)
+		else 
+#endif //!PLATFORM_NINTENDO
+		if (PrimaryPresence && PrimaryPresence->bIsOnline)
 		{
 			// We have no platform presence, but we do have primary, so get the platform from that
 			PlatformString = PrimaryPresence->GetPlatform();
diff --git a/Engine/Plugins/Online/OnlineSubsystem/Source/Private/OnlineSubsystem.cpp b/Engine/Plugins/Online/OnlineSubsystem/Source/Private/OnlineSubsystem.cpp
--- a/Engine/Plugins/Online/OnlineSubsystem/Source/Private/OnlineSubsystem.cpp
+++ b/Engine/Plugins/Online/OnlineSubsystem/Source/Private/OnlineSubsystem.cpp
@@ -266,7 +266,7 @@ TAutoConsoleVariable<FString> CVarPlatformOverride(
 	TEXT("oss.PlatformOverride"),
 	TEXT(""),
 	TEXT("Overrides the detected platform of this client for various debugging\n")
-	TEXT("Valid values WIN MAC PSN XBL IOS AND LIN SWT OTHER"),
+	TEXT("Valid values WIN MAC PSN XBL IOS AND LIN SWT SWT2 OTHER"),
 	ECVF_Cheat);
 
 FString IOnlineSubsystem::GetLocalPlatformName()
diff --git a/Engine/Plugins/Online/OnlineSubsystem/Source/Public/OnlineSubsystemNames.h b/Engine/Plugins/Online/OnlineSubsystem/Source/Public/OnlineSubsystemNames.h
--- a/Engine/Plugins/Online/OnlineSubsystem/Source/Public/OnlineSubsystemNames.h
+++ b/Engine/Plugins/Online/OnlineSubsystem/Source/Public/OnlineSubsystemNames.h
@@ -13,6 +13,7 @@
 #define OSS_PLATFORM_NAME_IOS		TEXT("IOS")
 #define OSS_PLATFORM_NAME_ANDROID	TEXT("AND")
 #define OSS_PLATFORM_NAME_SWITCH	TEXT("SWT")
+#define OSS_PLATFORM_NAME_SWITCH2	TEXT("SWT2")
 #define OSS_PLATFORM_NAME_OTHER		TEXT("OTHER")
 
 
