[QoS] GetSubregionPreferences - Return all subregions sorted by ping.

diff --git a/Engine/Plugins/Online/OnlineFramework/Source/Qos/Private/QosRegionManager.cpp b/Engine/Plugins/Online/OnlineFramework/Source/Qos/Private/QosRegionManager.cpp
--- a/Engine/Plugins/Online/OnlineFramework/Source/Qos/Private/QosRegionManager.cpp
+++ b/Engine/Plugins/Online/OnlineFramework/Source/Qos/Private/QosRegionManager.cpp
@@ -1,6 +1,8 @@
 // Copyright Epic Games, Inc. All Rights Reserved.
 
 #include "QosRegionManager.h"
+
+#include "Algo/Transform.h"
 #include "Misc/CommandLine.h"
 #include "Misc/ConfigCacheIni.h"
 #include "Misc/TrackedActivity.h"
@@ -14,6 +16,17 @@
 
 #define LAST_REGION_EVALUATION 3
 
+namespace QosCvars
+{
+	static bool bIncludeAllSubregionPreferences = true;
+	static FAutoConsoleVariableRef CVarbIncludeAllSubregionPreferences(
+		TEXT("QOS.IncludeAllSubregionPreferences"),
+		bIncludeAllSubregionPreferences,
+		TEXT("Whether requesting subregions includes all subregions regardless of parent region.\n")
+		TEXT("true Enables. false disables."),
+		ECVF_Default);
+}
+
 bool FDatacenterQosInstance::IsLessWhenBiasedTowardsNonSubspace(
 	const FDatacenterQosInstance& A,
 	const FDatacenterQosInstance& B,
@@ -254,6 +267,7 @@ void UQosRegionManager::PostReloadConfig(FProperty* PropertyThatWasLoaded)
 			TrySetDefaultRegion();
 		}
 
+		RefreshUsableSubregions();
 		SanityCheckDefinitions();
 	}
 }
@@ -289,6 +303,36 @@ bool UQosRegionManager::IsSubspaceBiasOrderEnabled(const FQosRegionInfo& RegionD
 	return IsSubspaceBiasOrderEnabled() && RegionDefinition.bAllowSubspaceBias;
 }
 
+void UQosRegionManager::RefreshUsableSubregions()
+{
+	TArray<const FDatacenterQosInstance*> UsableDatacentersByPing;
+	UsableDatacentersByPing.Reserve(DatacenterDefinitions.Num());
+	for (const FRegionQosInstance& RegionInstance : RegionOptions)
+	{
+		if (RegionInstance.IsUsable())
+		{
+			for (const FDatacenterQosInstance& DatacenterInstance : RegionInstance.DatacenterOptions)
+			{
+				UsableDatacentersByPing.Add(&DatacenterInstance);
+			}
+		}
+	}
+
+	UsableDatacentersByPing.Sort(
+		[&SubspaceBiasParams = SubspaceBiasParams, SubspaceDelimiter = GetSubspaceDelimiter()](const FDatacenterQosInstance& A, const FDatacenterQosInstance& B) {
+			return FDatacenterQosInstance::IsLessWhenBiasedTowardsNonSubspace(A, B, SubspaceBiasParams, SubspaceDelimiter);
+		});
+
+	UsableSubregions.Empty(UsableDatacentersByPing.Num());
+	Algo::Transform(UsableDatacentersByPing, UsableSubregions,
+		[](const FDatacenterQosInstance* Datacenter)
+		{
+			return Datacenter->Definition.Id;
+		});
+
+	UE_LOG(LogQos, Log, TEXT("[UQosRegionManager::RefreshUsableSubregions] Refreshed subreion list. [%s]"), *FString::Join(UsableSubregions, TEXT(",")));
+}
+
 // static
 FString UQosRegionManager::GetDatacenterId()
 {
@@ -475,6 +519,7 @@ void UQosRegionManager::OnQosEvaluationComplete(EQosCompletionResult Result, con
 		}
 	}
 
+	RefreshUsableSubregions();
 	LastCheckTimestamp = FDateTime::UtcNow();
 
 	if (!SelectedRegionId.IsEmpty() && SelectedRegionId == NO_REGION)
@@ -608,6 +653,12 @@ FString UQosRegionManager::GetBestRegion() const
 
 void UQosRegionManager::GetSubregionPreferences(const FString& RegionId, TArray<FString>& OutSubregions) const
 {
+	if (QosCvars::bIncludeAllSubregionPreferences && !bRegionForcedViaCommandline)
+	{
+		OutSubregions = UsableSubregions;
+		return;
+	}
+
 	const TArray<FRegionQosInstance>& LocalRegionOptions = GetRegionOptions();
 	for (const FRegionQosInstance& Region : LocalRegionOptions)
 	{
diff --git a/Engine/Plugins/Online/OnlineFramework/Source/Qos/Public/QosRegionManager.h b/Engine/Plugins/Online/OnlineFramework/Source/Qos/Public/QosRegionManager.h
--- a/Engine/Plugins/Online/OnlineFramework/Source/Qos/Public/QosRegionManager.h
+++ b/Engine/Plugins/Online/OnlineFramework/Source/Qos/Public/QosRegionManager.h
@@ -638,6 +638,11 @@ private:
 	 */
 	UE_API bool IsSubspaceBiasOrderEnabled(const FQosRegionInfo& RegionDefinition) const;
 
+	/**
+	 * Rebuild the list of usable subregions sorted by ping ascending.
+	 */
+	void RefreshUsableSubregions();
+
 	/** Number of times to ping a given region using random sampling of available servers */
 	UPROPERTY(Config)
 	int32 NumTestsPerRegion;
@@ -660,6 +665,10 @@ private:
 	UPROPERTY(Config)
 	FString SubspaceDelimiter;
 
+	/** Granular settings for biased subspace-based sorting algorithm which applies when returning all subregions for queries */
+	UPROPERTY(Config)
+	FQosSubspaceComparisonParams SubspaceBiasParams;
+
 	/** Metadata about existing regions */
 	UPROPERTY(Config)
 	TArray<FQosRegionInfo> RegionDefinitions;
@@ -690,6 +699,8 @@ private:
 	/** Value set by the game to be the current region */
 	UPROPERTY()
 	FString SelectedRegionId;
+	/** List of all useable subregions sorted by ping. Does not include accelerated regions. */
+	TArray<FString> UsableSubregions;
 
 	FOnQosEvalCompleteDelegate OnQosEvalCompleteDelegate;
 
