[Hotfix] Fix hotfixing of dynamic ini files when additional hotfixes are applied.
Change per-row asset hotfix logging to VeryVerbose.

diff --git a/Engine/Plugins/Online/OnlineFramework/Source/Hotfix/Private/OnlineHotfixManager.cpp b/Engine/Plugins/Online/OnlineFramework/Source/Hotfix/Private/OnlineHotfixManager.cpp
--- a/Engine/Plugins/Online/OnlineFramework/Source/Hotfix/Private/OnlineHotfixManager.cpp
+++ b/Engine/Plugins/Online/OnlineFramework/Source/Hotfix/Private/OnlineHotfixManager.cpp
@@ -47,6 +47,12 @@ static TAutoConsoleVariable<int32> CVarUseNewDynamicLayersForHotfix(
 	TEXT("If true, use the new dynamic layers that load/unload configs, specifically for Hotfixes"),
 	ECVF_Default);
 
+static TAutoConsoleVariable<bool> CVarEmptyDynamicHotfixContentsOnNewHotfix(
+	TEXT("hotfix.EmptyDynamicHotfixContentsOnNewHotfix"),
+	false,
+	TEXT("If true clear out the existing values for DynamicHotfixContents when a new hotfix is processed. When set to false only remove the files which were removed since the last hotfix."),
+	ECVF_Default);
+
 
 class FPakFileVisitor : public IPlatformFile::FDirectoryVisitor
 {
@@ -755,15 +761,29 @@ void UOnlineHotfixManager::UpdateProgress(uint32 FileCount, uint64 UpdateSize)
 
 EHotfixResult UOnlineHotfixManager::ApplyHotfix()
 {
+	// Check the CVar value before unmounting any previous hotfixes in case the value was changed by the previous hotfix.
+	const bool bEmptyDynamicHotfixContentsOnNewHotfix = CVarEmptyDynamicHotfixContentsOnNewHotfix->GetBool();
+	if (bEmptyDynamicHotfixContentsOnNewHotfix)
+	{
+		UE_LOG(LogHotfixManager, Log, TEXT("[DynamicHotfixContents] Removing ini data for all pending hotfixes"));
+		DynamicHotfixContents.Empty();
+	}
+
 	// Perform any undo operations needed
 	// This occurs same frame as the application of new hotfixes
 	RestoreBackupIniFiles();
 	UnmountHotfixFiles();
-	DynamicHotfixContents.Empty();
 
 	for (const FCloudFileHeader& FileHeader : RemovedHotfixFileList)
 	{
 		TriggerOnHotfixRemovedFileDelegates(FileHeader.FileName);
+
+		if (!bEmptyDynamicHotfixContentsOnNewHotfix)
+		{
+			const FString HotfixTag = GetStrippedConfigFileName(FileHeader.FileName).Replace(TEXT(".ini"), TEXT(""));
+			UE_LOG(LogHotfixManager, Log, TEXT("[DynamicHotfixContents] Removing ini data for pending hotfix %s (%s)"), *FileHeader.FileName, *HotfixTag);
+			DynamicHotfixContents.Remove(FName(HotfixTag));
+		}
 	}
 
 	for (const FCloudFileHeader& FileHeader : ChangedHotfixFileList)
@@ -1032,7 +1052,7 @@ bool UOnlineHotfixManager::HotfixIniFile(const FString& FileName, const FString&
 			{
 				const FString HotfixTag = GetStrippedConfigFileName(FileName).Replace(TEXT(".ini"), TEXT(""));
 
-				UE_LOG(LogHotfixManager, Log, TEXT("Storing ini data for pending hotfix %s (%s)"), *FileName, *HotfixTag);
+				UE_LOG(LogHotfixManager, Log, TEXT("[DynamicHotfixContents] Storing ini data for pending hotfix %s (%s)"), *FileName, *HotfixTag);
 
 				TArray<TPair<FString, FString>>& HotfixesForTag = DynamicHotfixContents.FindOrAdd(*HotfixTag);
 				HotfixesForTag.Add(TPair<FString, FString>(FileName, IniData));
@@ -1762,7 +1782,7 @@ void UOnlineHotfixManager::HotfixRowUpdate(
 								NumProp->SetIntPropertyValue(RowData, NewPropertyValue);
 								OnHotfixTableValueInt64(*Asset, RowName, ColumnName, OldPropertyValue, NewPropertyValue);
 								bWasDataTableChanged = true;
-								UE_LOG(LogHotfixManager, Log, TEXT("Data table %s row %s updated column %s from %" INT64_FMT " to %" INT64_FMT "."), *AssetPath, *RowName, *ColumnName, OldPropertyValue, NewPropertyValue);
+								UE_LOG(LogHotfixManager, VeryVerbose, TEXT("Data table %s row %s updated column %s from %" INT64_FMT " to %" INT64_FMT "."), *AssetPath, *RowName, *ColumnName, OldPropertyValue, NewPropertyValue);
 							}
 							// Float
 							else
@@ -1772,7 +1792,7 @@ void UOnlineHotfixManager::HotfixRowUpdate(
 								NumProp->SetFloatingPointPropertyValue(RowData, NewPropertyValue);
 								OnHotfixTableValueDouble(*Asset, RowName, ColumnName, OldPropertyValue, NewPropertyValue);
 								bWasDataTableChanged = true;
-								UE_LOG(LogHotfixManager, Log, TEXT("Data table %s row %s updated column %s from %.2f to %.2f."), *AssetPath, *RowName, *ColumnName, OldPropertyValue, NewPropertyValue);
+								UE_LOG(LogHotfixManager, VeryVerbose, TEXT("Data table %s row %s updated column %s from %.2f to %.2f."), *AssetPath, *RowName, *ColumnName, OldPropertyValue, NewPropertyValue);
 							}
 						}
 						// Not a number.
@@ -1790,7 +1810,7 @@ void UOnlineHotfixManager::HotfixRowUpdate(
 						StrProp->SetPropertyValue(RowData, NewPropertyValue);
 						OnHotfixTableValueString(*Asset, RowName, ColumnName, OldPropertyValue, NewPropertyValue);
 						bWasDataTableChanged = true;
-						UE_LOG(LogHotfixManager, Log, TEXT("Data table %s row %s updated column %s from %s to %s."), *AssetPath, *RowName, *ColumnName, *OldPropertyValue, *NewPropertyValue);
+						UE_LOG(LogHotfixManager, VeryVerbose, TEXT("Data table %s row %s updated column %s from %s to %s."), *AssetPath, *RowName, *ColumnName, *OldPropertyValue, *NewPropertyValue);
 					}
 					// FName property
 					else if (NameProp)
@@ -1800,7 +1820,7 @@ void UOnlineHotfixManager::HotfixRowUpdate(
 						NameProp->SetPropertyValue(RowData, NewPropertyValue);
 						OnHotfixTableValueName(*Asset, RowName, ColumnName, OldPropertyValue, NewPropertyValue);
 						bWasDataTableChanged = true;
-						UE_LOG(LogHotfixManager, Log, TEXT("Data table %s row %s updated column %s from %s to %s."), *AssetPath, *RowName, *ColumnName, *OldPropertyValue.ToString(), *NewPropertyValue.ToString());
+						UE_LOG(LogHotfixManager, VeryVerbose, TEXT("Data table %s row %s updated column %s from %s to %s."), *AssetPath, *RowName, *ColumnName, *OldPropertyValue.ToString(), *NewPropertyValue.ToString());
 					}
 					// Soft Object property
 					else if (SoftObjProp)
@@ -1810,7 +1830,7 @@ void UOnlineHotfixManager::HotfixRowUpdate(
 						SoftObjProp->SetPropertyValue(RowData, NewPropertyValue);
 						OnHotfixTableValueSoftObject(*Asset, RowName, ColumnName, OldPropertyValue, NewPropertyValue);
 						bWasDataTableChanged = true;
-						UE_LOG(LogHotfixManager, Log, TEXT("Data table %s row %s updated column %s from %s to %s."), *AssetPath, *RowName, *ColumnName, *OldPropertyValue.ToString(), *NewPropertyValue.ToString());
+						UE_LOG(LogHotfixManager, VeryVerbose, TEXT("Data table %s row %s updated column %s from %s to %s."), *AssetPath, *RowName, *ColumnName, *OldPropertyValue.ToString(), *NewPropertyValue.ToString());
 					}
 					else
 					{
@@ -1821,7 +1841,7 @@ void UOnlineHotfixManager::HotfixRowUpdate(
 						if (Error.IsEmpty())
 						{
 							bWasDataTableChanged = true;
-							UE_LOG(LogHotfixManager, Log, TEXT("Data table %s row %s updated column %s from %s to %s."), *AssetPath, *RowName, *ColumnName, *OldPropertyValue, *NewValue);
+							UE_LOG(LogHotfixManager, VeryVerbose, TEXT("Data table %s row %s updated column %s from %s to %s."), *AssetPath, *RowName, *ColumnName, *OldPropertyValue, *NewValue);
 						}
 						else
 						{
@@ -1894,11 +1914,11 @@ void UOnlineHotfixManager::HotfixRowUpdate(
 
 						if (bWasExistingKey)
 						{
-							UE_LOG(LogHotfixManager, Log, TEXT("Curve table %s row %s updated column %s from %.2f to %.2f."), *AssetPath, *RowName, *ColumnName, OldPropertyValue, NewPropertyValue);
+							UE_LOG(LogHotfixManager, VeryVerbose, TEXT("Curve table %s row %s updated column %s from %.2f to %.2f."), *AssetPath, *RowName, *ColumnName, OldPropertyValue, NewPropertyValue);
 						}
 						else
 						{
-							UE_LOG(LogHotfixManager, Log, TEXT("Curve table %s row %s added column %s with value %.2f."), *AssetPath, *RowName, *ColumnName, NewPropertyValue);
+							UE_LOG(LogHotfixManager, VeryVerbose, TEXT("Curve table %s row %s added column %s with value %.2f."), *AssetPath, *RowName, *ColumnName, NewPropertyValue);
 						}
 					}
 					else
@@ -1953,7 +1973,7 @@ void UOnlineHotfixManager::HotfixRowUpdate(
 					CurveFloat->FloatCurve.SetKeyValue(Key, NewPropertyValue);
 					OnHotfixTableValueFloat(*Asset, RowName, ColumnName, OldPropertyValue, NewPropertyValue);
 
-					UE_LOG(LogHotfixManager, Log, TEXT("Curve float %s updated column %s from %.2f to %.2f."), *AssetPath, *ColumnName, OldPropertyValue, NewPropertyValue);
+					UE_LOG(LogHotfixManager, VeryVerbose, TEXT("Curve float %s updated column %s from %.2f to %.2f."), *AssetPath, *ColumnName, OldPropertyValue, NewPropertyValue);
 				}
 				else
 				{
@@ -2100,18 +2120,18 @@ void UOnlineHotfixManager::HotfixDynamicBranch(const FName& Tag, const FName& Br
 	const FString HotfixBaseName = Tag.ToString() + Branch.ToString();
 	const TArray<TPair<FString, FString>>* HotfixesForTag = DynamicHotfixContents.Find(*HotfixBaseName);
 
-	UE_LOG(LogHotfixManager, VeryVerbose, TEXT("HotfixDynamicBranch: Searching for hotfix for base name: %s. Result: %s"), *HotfixBaseName, HotfixesForTag ? TEXT("found") : TEXT("not found"));
+	UE_LOG(LogHotfixManager, VeryVerbose, TEXT("[DynamicHotfixContents] HotfixDynamicBranch: Searching for hotfix for base name: %s. Result: %s"), *HotfixBaseName, HotfixesForTag ? TEXT("found") : TEXT("not found"));
 
 	if (HotfixesForTag)
 	{
 		for (const TPair<FString, FString>& HotfixFilenameAndContents : *HotfixesForTag)
 		{
-			UE_LOG(LogHotfixManager, Log, TEXT("HotfixDynamicBranch: applying hotfix %s for tag %s on branch %s"), *HotfixFilenameAndContents.Key, *Tag.ToString(), *Branch.ToString());
+			UE_LOG(LogHotfixManager, Log, TEXT("[DynamicHotfixContents] HotfixDynamicBranch: applying hotfix %s for tag %s on branch %s"), *HotfixFilenameAndContents.Key, *Tag.ToString(), *Branch.ToString());
 
 			FConfigBranch* BranchToHotfix = GConfig->FindBranch(Branch, FString());
 			if (BranchToHotfix && !BranchToHotfix->AddDynamicLayerStringToHierarchy(HotfixFilenameAndContents.Key, HotfixFilenameAndContents.Value, Tag, DynamicLayerPriority::Hotfix, ModificationTracker))
 			{
-				UE_LOG(LogHotfixManager, Warning, TEXT("HotfixDynamicBranch: failed to apply hotfix"));
+				UE_LOG(LogHotfixManager, Warning, TEXT("[DynamicHotfixContents] HotfixDynamicBranch: failed to apply hotfix"));
 			}
 		}
 	}
