Replaced bool parameters with EFindObjectFlags.

diff --git a/Engine/Plugins/Editor/DataValidation/Source/DataValidation/Private/WorldPartitionChangelistValidator.cpp b/Engine/Plugins/Editor/DataValidation/Source/DataValidation/Private/WorldPartitionChangelistValidator.cpp
--- a/Engine/Plugins/Editor/DataValidation/Source/DataValidation/Private/WorldPartitionChangelistValidator.cpp
+++ b/Engine/Plugins/Editor/DataValidation/Source/DataValidation/Private/WorldPartitionChangelistValidator.cpp
@@ -201,7 +201,7 @@ void UWorldPartitionChangelistValidator::ValidateActorsAndDataLayersFromChangeLi
 		const TSet<FAssetData>& ActorsData = It.Get<1>();
 
 		// Find/Load the ActorDescContainer
-		UWorld* World = FindObject<UWorld>(nullptr, *MapPath.ToString(), true);
+		UWorld* World = FindObject<UWorld>(nullptr, *MapPath.ToString(), EFindObjectFlags::ExactClass);
 		
 		TGuardValue<UObject*> GuardCurrentAsset(CurrentAsset, World);
 		
diff --git a/Engine/Plugins/FX/Niagara/Source/Niagara/Private/NiagaraDebugHud.cpp b/Engine/Plugins/FX/Niagara/Source/Niagara/Private/NiagaraDebugHud.cpp
--- a/Engine/Plugins/FX/Niagara/Source/Niagara/Private/NiagaraDebugHud.cpp
+++ b/Engine/Plugins/FX/Niagara/Source/Niagara/Private/NiagaraDebugHud.cpp
@@ -2056,13 +2056,13 @@ void FNiagaraDebugHud::DrawOverview(class FNiagaraWorldManager* WorldManager, FC
 			DrawCanvas->DrawShadowedString(TextLocation.X, TextLocation.Y, TEXT("Niagara DebugHud"), Font, HeadingColor);
 			TextLocation.Y += fAdvanceHeight;
 
-			const UEnum* EnumPtr = FindObjectChecked<UEnum>(nullptr, TEXT("/Script/Niagara.ENiagaraDebugHUDOverviewMode"), true);
+			const UEnum* EnumPtr = FindObjectChecked<UEnum>(nullptr, TEXT("/Script/Niagara.ENiagaraDebugHUDOverviewMode"), EFindObjectFlags::ExactClass);
 			TStringBuilder<128> ModelStringBuilder;
 			ModelStringBuilder.Append(EnumPtr->GetDisplayNameTextByValue((int32)Settings.OverviewMode).ToString());
 			if(Settings.OverviewMode == ENiagaraDebugHUDOverviewMode::PerformanceGraph)
 			{
-				const UEnum* PerfGraphModeEnumPtr = FindObjectChecked<UEnum>(nullptr, TEXT("/Script/Niagara.ENiagaraDebugHUDPerfGraphMode"), true);
-				const UEnum* PerfGraphSampleModeEnumPtr = FindObjectChecked<UEnum>(nullptr, TEXT("/Script/Niagara.ENiagaraDebugHUDPerfSampleMode"), true);
+				const UEnum* PerfGraphModeEnumPtr = FindObjectChecked<UEnum>(nullptr, TEXT("/Script/Niagara.ENiagaraDebugHUDPerfGraphMode"), EFindObjectFlags::ExactClass);
+				const UEnum* PerfGraphSampleModeEnumPtr = FindObjectChecked<UEnum>(nullptr, TEXT("/Script/Niagara.ENiagaraDebugHUDPerfSampleMode"), EFindObjectFlags::ExactClass);
 				
 				ModelStringBuilder.Append(TEXT(" - "));
 				ModelStringBuilder.Append(PerfGraphModeEnumPtr->GetDisplayNameTextByValue((int32)Settings.PerfGraphMode).ToString());
diff --git a/Engine/Plugins/Online/OnlineFramework/Source/Hotfix/Private/OnlineHotfixManager.cpp b/Engine/Plugins/Online/OnlineFramework/Source/Hotfix/Private/OnlineHotfixManager.cpp
--- a/Engine/Plugins/Online/OnlineFramework/Source/Hotfix/Private/OnlineHotfixManager.cpp
+++ b/Engine/Plugins/Online/OnlineFramework/Source/Hotfix/Private/OnlineHotfixManager.cpp
@@ -1144,7 +1144,7 @@ bool UOnlineHotfixManager::IsMapLoaded(const FString& MapName)
 	FString MapPackageName(MapName.Left(MapName.Len() - 5));
 	MapPackageName = MapPackageName.Replace(*GameContentPath, TEXT("/Game"));
 	// If this map's UPackage exists, it is currently loaded
-	UPackage* MapPackage = FindObject<UPackage>(nullptr, *MapPackageName, true);
+	UPackage* MapPackage = FindObject<UPackage>(nullptr, *MapPackageName, EFindObjectFlags::ExactClass);
 	return MapPackage != nullptr;
 }
 
@@ -1379,7 +1379,7 @@ void UOnlineHotfixManager::RestoreBackupIniFiles()
 		RestoredClasses.Reserve(ClassesToRestore.Num());
 		for (int32 Index = 0; Index < ClassesToRestore.Num(); Index++)
 		{
-			UClass* Class = FindObject<UClass>(nullptr, *ClassesToRestore[Index], true);
+			UClass* Class = FindObject<UClass>(nullptr, *ClassesToRestore[Index], EFindObjectFlags::ExactClass);
 			if (Class != nullptr)
 			{
 				// Add this to the list to check against
@@ -1682,7 +1682,7 @@ void UOnlineHotfixManager::ReloadConfigsFromIniFiles()
 		RestoredClasses.Reserve(ClassesToReload.Num());
 		for (int32 Index = 0; Index < ClassesToReload.Num(); Index++)
 		{
-			UClass* Class = FindObject<UClass>(nullptr, *ClassesToReload[Index], true);
+			UClass* Class = FindObject<UClass>(nullptr, *ClassesToReload[Index], EFindObjectFlags::ExactClass);
 			if (Class != nullptr)
 			{
 				// Add this to the list to check against
@@ -2207,7 +2207,7 @@ void UOnlineHotfixManager::ReloadObjectsAffectedByConfigFile(const FString& IniD
 						// Snip the text out and try to find the class for that
 						const FString PackageClassName = IniData.Mid(StartIndex + ScriptSectionTag, EndIndex - StartIndex - ScriptSectionTag);
 						// Find the class for this so we know what to update
-						UClass* Class = FindObject<UClass>(nullptr, *PackageClassName, true);
+						UClass* Class = FindObject<UClass>(nullptr, *PackageClassName, EFindObjectFlags::ExactClass);
 						if (Class)
 						{
 							// Add this to the list to check against
diff --git a/Engine/Plugins/Runtime/GameplayAbilities/Source/GameplayAbilities/Private/AbilitySystemCheatManagerExtension.cpp b/Engine/Plugins/Runtime/GameplayAbilities/Source/GameplayAbilities/Private/AbilitySystemCheatManagerExtension.cpp
--- a/Engine/Plugins/Runtime/GameplayAbilities/Source/GameplayAbilities/Private/AbilitySystemCheatManagerExtension.cpp
+++ b/Engine/Plugins/Runtime/GameplayAbilities/Source/GameplayAbilities/Private/AbilitySystemCheatManagerExtension.cpp
@@ -212,10 +212,10 @@ void UAbilitySystemCheatManagerExtension::AbilityListGranted() const
 		return;
 	}
 
-	const UEnum* ExecutionEnumPtr = FindObject<UEnum>(nullptr, TEXT("/Script/GameplayAbilities.EGameplayAbilityNetExecutionPolicy"), true);
+	const UEnum* ExecutionEnumPtr = FindObject<UEnum>(nullptr, TEXT("/Script/GameplayAbilities.EGameplayAbilityNetExecutionPolicy"), EFindObjectFlags::ExactClass);
 	check(ExecutionEnumPtr && TEXT("Couldn't locate EGameplayAbilityNetExecutionPolicy enum!"));
 
-	const UEnum* SecurityEnumPtr = FindObject<UEnum>(nullptr, TEXT("/Script/GameplayAbilities.EGameplayAbilityNetSecurityPolicy"), true);
+	const UEnum* SecurityEnumPtr = FindObject<UEnum>(nullptr, TEXT("/Script/GameplayAbilities.EGameplayAbilityNetSecurityPolicy"), EFindObjectFlags::ExactClass);
 	check(SecurityEnumPtr && TEXT("Couldn't locate EGameplayAbilityNetSecurityPolicy enum!"));
 
 	UE_LOG(LogConsoleResponse, Log, TEXT("Granted abilities to %s (ASC: '%s'):"), *PC->GetName(), *ASC->GetFullName());
diff --git a/Engine/Plugins/Runtime/RigVM/Source/RigVM/Private/RigVMCore/RigVMMemoryStorage.cpp b/Engine/Plugins/Runtime/RigVM/Source/RigVM/Private/RigVMCore/RigVMMemoryStorage.cpp
--- a/Engine/Plugins/Runtime/RigVM/Source/RigVM/Private/RigVMCore/RigVMMemoryStorage.cpp
+++ b/Engine/Plugins/Runtime/RigVM/Source/RigVM/Private/RigVMCore/RigVMMemoryStorage.cpp
@@ -511,7 +511,7 @@ URigVMMemoryStorageGeneratorClass* URigVMMemoryStorageGeneratorClass::CreateStor
 	// does not guarantee that condition. As a result we will do OldClass->ConditionalPostLoad() to ensure the OldClass
 	// is fully loaded, ready to be replaced.
 	
-	UObject* OldClassObject = StaticFindObjectFastInternal( /*Class=*/ NULL, Package, *ClassName, true, RF_NoFlags, EInternalObjectFlags::None);
+	UObject* OldClassObject = StaticFindObjectFastInternal( /*Class=*/ NULL, Package, *ClassName, EFindObjectFlags::ExactClass, RF_NoFlags, EInternalObjectFlags::None);
 	if(OldClassObject)
 	{
 		auto FindUniqueName = [](const TCHAR* TemplateName, int32& TemplateIndex)
diff --git a/Engine/Plugins/Runtime/RigVM/Source/RigVM/Private/RigVMCore/RigVMObjectArchive.cpp b/Engine/Plugins/Runtime/RigVM/Source/RigVM/Private/RigVMCore/RigVMObjectArchive.cpp
--- a/Engine/Plugins/Runtime/RigVM/Source/RigVM/Private/RigVMCore/RigVMObjectArchive.cpp
+++ b/Engine/Plugins/Runtime/RigVM/Source/RigVM/Private/RigVMCore/RigVMObjectArchive.cpp
@@ -327,7 +327,7 @@ FArchive& FRigVMObjectArchiveReader::operator<<(UObject*& Obj)
 			}
 			else
 			{
-				Obj = FindObject<UObject>(nullptr, *PathName.ToString(), false);
+				Obj = FindObject<UObject>(nullptr, *PathName.ToString(), EFindObjectFlags::None);
 				check(Obj);
 				ReadObjects.Add(PathName, Obj);
 			}
@@ -428,7 +428,7 @@ FRigVMObjectArchiveReader::FObjectHeader FRigVMObjectArchiveReader::ReadObjectHe
 	}
 	else
 	{
-		Header.Class = FindObject<UClass>(nullptr, *ClassName.ToString(), false);
+		Header.Class = FindObject<UClass>(nullptr, *ClassName.ToString(), EFindObjectFlags::None);
 		ReadClasses.Add(ClassName, Header.Class);
 	}
 
diff --git a/Engine/Source/Developer/AssetTools/Private/AssetTools.cpp b/Engine/Source/Developer/AssetTools/Private/AssetTools.cpp
--- a/Engine/Source/Developer/AssetTools/Private/AssetTools.cpp
+++ b/Engine/Source/Developer/AssetTools/Private/AssetTools.cpp
@@ -1184,7 +1184,7 @@ namespace UE::AssetTools::Private
 
 			for (const FName PackageName : PackagesToClean)
 			{
-				if (UPackage* Package = FindObjectFast<UPackage>(nullptr, PackageName, true))
+				if (UPackage* Package = FindObjectFast<UPackage>(nullptr, PackageName, EFindObjectFlags::ExactClass))
 				{
 					ExistingPackageToClean.Add(Package);
 				}
diff --git a/Engine/Source/Editor/AssetDefinition/Private/AssetStatusAssetDataInfoProvider.cpp b/Engine/Source/Editor/AssetDefinition/Private/AssetStatusAssetDataInfoProvider.cpp
--- a/Engine/Source/Editor/AssetDefinition/Private/AssetStatusAssetDataInfoProvider.cpp
+++ b/Engine/Source/Editor/AssetDefinition/Private/AssetStatusAssetDataInfoProvider.cpp
@@ -1,4 +1,4 @@
-﻿// Copyright Epic Games, Inc. All Rights Reserved.
+// Copyright Epic Games, Inc. All Rights Reserved.
 
 #include "AssetStatusAssetDataInfoProvider.h"
 
@@ -8,7 +8,7 @@
 
 UPackage* FAssetStatusAssetDataInfoProvider::FindPackage() const
 {
-	return FindObjectSafe<UPackage>(NULL, *FNameBuilder(AssetData.PackageName), true);
+	return FindObjectSafe<UPackage>(NULL, *FNameBuilder(AssetData.PackageName), EFindObjectFlags::ExactClass);
 }
 
 FString FAssetStatusAssetDataInfoProvider::TryGetFilename() const
diff --git a/Engine/Source/Editor/BlueprintGraph/Private/EdGraphSchema_K2.cpp b/Engine/Source/Editor/BlueprintGraph/Private/EdGraphSchema_K2.cpp
--- a/Engine/Source/Editor/BlueprintGraph/Private/EdGraphSchema_K2.cpp
+++ b/Engine/Source/Editor/BlueprintGraph/Private/EdGraphSchema_K2.cpp
@@ -1388,7 +1388,7 @@ bool UEdGraphSchema_K2::PinHasSplittableStructType(const UEdGraphPin* InGraphPin
 				if (!bCanSplit)
 				{
 					const FString& MetaData = StructType->GetMetaData(FBlueprintMetadata::MD_NativeMakeFunction);
-					UFunction* Function = FindObject<UFunction>(nullptr, *MetaData, true);
+					UFunction* Function = FindObject<UFunction>(nullptr, *MetaData, EFindObjectFlags::ExactClass);
 					bCanSplit = (Function != nullptr);
 				}
 			}
@@ -1398,7 +1398,7 @@ bool UEdGraphSchema_K2::PinHasSplittableStructType(const UEdGraphPin* InGraphPin
 				if (!bCanSplit)
 				{
 					const FString& MetaData = StructType->GetMetaData(FBlueprintMetadata::MD_NativeBreakFunction);
-					UFunction* Function = FindObject<UFunction>(nullptr, *MetaData, true);
+					UFunction* Function = FindObject<UFunction>(nullptr, *MetaData, EFindObjectFlags::ExactClass);
 					bCanSplit = (Function != nullptr);
 				}
 			}
@@ -7259,7 +7259,7 @@ UK2Node* UEdGraphSchema_K2::CreateSplitPinNode(UEdGraphPin* Pin, const FCreateSp
 		else
 		{
 			const FString& MetaData = StructType->GetMetaData(FBlueprintMetadata::MD_NativeMakeFunction);
-			const UFunction* Function = FindObject<UFunction>(nullptr, *MetaData, true);
+			const UFunction* Function = FindObject<UFunction>(nullptr, *MetaData, EFindObjectFlags::ExactClass);
 
 			UK2Node_CallFunction* CallFunctionNode;
 
@@ -7307,7 +7307,7 @@ UK2Node* UEdGraphSchema_K2::CreateSplitPinNode(UEdGraphPin* Pin, const FCreateSp
 		else
 		{
 			const FString& MetaData = StructType->GetMetaData(FBlueprintMetadata::MD_NativeBreakFunction);
-			const UFunction* Function = FindObject<UFunction>(nullptr, *MetaData, true);
+			const UFunction* Function = FindObject<UFunction>(nullptr, *MetaData, EFindObjectFlags::ExactClass);
 
 			UK2Node_CallFunction* CallFunctionNode;
 
diff --git a/Engine/Source/Editor/BlueprintGraph/Private/K2Node_BreakStruct.cpp b/Engine/Source/Editor/BlueprintGraph/Private/K2Node_BreakStruct.cpp
--- a/Engine/Source/Editor/BlueprintGraph/Private/K2Node_BreakStruct.cpp
+++ b/Engine/Source/Editor/BlueprintGraph/Private/K2Node_BreakStruct.cpp
@@ -451,7 +451,7 @@ void UK2Node_BreakStruct::ConvertDeprecatedNode(UEdGraph* Graph, bool bOnlySafeC
 		else
 		{
 			const FString& MetaData = StructType->GetMetaData(FBlueprintMetadata::MD_NativeBreakFunction);
-			BreakNodeFunction = FindObject<UFunction>(nullptr, *MetaData, true);
+			BreakNodeFunction = FindObject<UFunction>(nullptr, *MetaData, EFindObjectFlags::ExactClass);
 			if (BreakNodeFunction)
 			{
 				// Look for the first parameter
diff --git a/Engine/Source/Editor/BlueprintGraph/Private/K2Node_FormatText.cpp b/Engine/Source/Editor/BlueprintGraph/Private/K2Node_FormatText.cpp
--- a/Engine/Source/Editor/BlueprintGraph/Private/K2Node_FormatText.cpp
+++ b/Engine/Source/Editor/BlueprintGraph/Private/K2Node_FormatText.cpp
@@ -452,7 +452,7 @@ void UK2Node_FormatText::ExpandNode(class FKismetCompilerContext& CompilerContex
 			}
 			else if (ArgumentPinCategory == UEdGraphSchema_K2::PC_Byte || ArgumentPinCategory == UEdGraphSchema_K2::PC_Enum)
 			{
-				static UEnum* TextGenderEnum = FindObjectChecked<UEnum>(nullptr, TEXT("/Script/Engine.ETextGender"), /*ExactClass*/true);
+				static UEnum* TextGenderEnum = FindObjectChecked<UEnum>(nullptr, TEXT("/Script/Engine.ETextGender"), EFindObjectFlags::ExactClass);
 				if (ArgumentPin->PinType.PinSubCategoryObject == TextGenderEnum)
 				{
 					MakeFormatArgumentDataStruct->GetSchema()->TrySetDefaultValue(*ArgumentTypePin, TEXT("Gender"));
@@ -583,12 +583,12 @@ bool UK2Node_FormatText::IsConnectionDisallowed(const UEdGraphPin* MyPin, const
 			(OtherPinCategory == UEdGraphSchema_K2::PC_Byte && !OtherPin->PinType.PinSubCategoryObject.IsValid()) ||
 			OtherPinCategory == UEdGraphSchema_K2::PC_Boolean || OtherPinCategory == UEdGraphSchema_K2::PC_String || OtherPinCategory == UEdGraphSchema_K2::PC_Name || OtherPinCategory == UEdGraphSchema_K2::PC_Object ||
 			OtherPinCategory == UEdGraphSchema_K2::PC_Wildcard || OtherPinCategory == UEdGraphSchema_K2::PC_Int64)
-        {
+		{
 			bIsValidType = true;
 		}
 		else if (OtherPinCategory == UEdGraphSchema_K2::PC_Byte || OtherPinCategory == UEdGraphSchema_K2::PC_Enum)
 		{
-			static UEnum* TextGenderEnum = FindObjectChecked<UEnum>(nullptr, TEXT("/Script/Engine.ETextGender"), /*ExactClass*/true);
+			static UEnum* TextGenderEnum = FindObjectChecked<UEnum>(nullptr, TEXT("/Script/Engine.ETextGender"), EFindObjectFlags::ExactClass);
 			if (OtherPin->PinType.PinSubCategoryObject == TextGenderEnum)
 			{
 				bIsValidType = true;
diff --git a/Engine/Source/Editor/BlueprintGraph/Private/K2Node_MakeStruct.cpp b/Engine/Source/Editor/BlueprintGraph/Private/K2Node_MakeStruct.cpp
--- a/Engine/Source/Editor/BlueprintGraph/Private/K2Node_MakeStruct.cpp
+++ b/Engine/Source/Editor/BlueprintGraph/Private/K2Node_MakeStruct.cpp
@@ -457,7 +457,7 @@ void UK2Node_MakeStruct::ConvertDeprecatedNode(UEdGraph* Graph, bool bOnlySafeCh
 		else
 		{
 			const FString& MetaData = StructType->GetMetaData(FBlueprintMetadata::MD_NativeMakeFunction);
-			MakeNodeFunction = FindObject<UFunction>(nullptr, *MetaData, true);
+			MakeNodeFunction = FindObject<UFunction>(nullptr, *MetaData, EFindObjectFlags::ExactClass);
 
 			if (MakeNodeFunction)
 			{
diff --git a/Engine/Source/Editor/BlueprintGraph/Private/K2Node_SpawnActorFromClass.cpp b/Engine/Source/Editor/BlueprintGraph/Private/K2Node_SpawnActorFromClass.cpp
--- a/Engine/Source/Editor/BlueprintGraph/Private/K2Node_SpawnActorFromClass.cpp
+++ b/Engine/Source/Editor/BlueprintGraph/Private/K2Node_SpawnActorFromClass.cpp
@@ -111,7 +111,7 @@ void UK2Node_SpawnActorFromClass::MaybeUpdateCollisionPin(TArray<UEdGraphPin*>&
 				UEdGraphPin* const CollisionHandlingOverridePin = GetCollisionHandlingOverridePin();
 				if (CollisionHandlingOverridePin)
 				{
-					UEnum const* const MethodEnum = FindObjectChecked<UEnum>(nullptr, TEXT("/Script/Engine.ESpawnActorCollisionHandlingMethod"), true);
+					UEnum const* const MethodEnum = FindObjectChecked<UEnum>(nullptr, TEXT("/Script/Engine.ESpawnActorCollisionHandlingMethod"), EFindObjectFlags::ExactClass);
 					CollisionHandlingOverridePin->DefaultValue =
 						bOldCollisionPinValue
 						? MethodEnum->GetNameStringByValue(static_cast<int>(ESpawnActorCollisionHandlingMethod::AlwaysSpawn))
@@ -126,7 +126,7 @@ void UK2Node_SpawnActorFromClass::MaybeUpdateCollisionPin(TArray<UEdGraphPin*>&
 				UEdGraphPin* const CollisionHandlingOverridePin = GetCollisionHandlingOverridePin();
 				check(CollisionHandlingOverridePin);
 
-				UEnum* const MethodEnum = FindObjectChecked<UEnum>(nullptr, TEXT("/Script/Engine.ESpawnActorCollisionHandlingMethod"), true);
+				UEnum* const MethodEnum = FindObjectChecked<UEnum>(nullptr, TEXT("/Script/Engine.ESpawnActorCollisionHandlingMethod"), EFindObjectFlags::ExactClass);
 				
 				FGraphNodeCreator<UK2Node_EnumLiteral> AlwaysSpawnLiteralCreator(*GetGraph());
 				UK2Node_EnumLiteral* const AlwaysSpawnLiteralNode = AlwaysSpawnLiteralCreator.CreateNode();
diff --git a/Engine/Source/Editor/Blutility/Private/EditorUtilityLibrary.cpp b/Engine/Source/Editor/Blutility/Private/EditorUtilityLibrary.cpp
--- a/Engine/Source/Editor/Blutility/Private/EditorUtilityLibrary.cpp
+++ b/Engine/Source/Editor/Blutility/Private/EditorUtilityLibrary.cpp
@@ -302,7 +302,7 @@ void UEditorUtilityLibrary::RenameAsset(UObject* Asset, const FString& NewName)
 AActor* UEditorUtilityLibrary::GetActorReference(FString PathToActor)
 {
 #if WITH_EDITOR
-	return Cast<AActor>(StaticFindObject(AActor::StaticClass(), GEditor->GetEditorWorldContext().World(), *PathToActor, false));
+	return Cast<AActor>(StaticFindObject(AActor::StaticClass(), GEditor->GetEditorWorldContext().World(), *PathToActor, EFindObjectFlags::None));
 #else
 	return nullptr;
 #endif //WITH_EDITOR
diff --git a/Engine/Source/Editor/Blutility/Private/GlobalEditorUtilityBase.cpp b/Engine/Source/Editor/Blutility/Private/GlobalEditorUtilityBase.cpp
--- a/Engine/Source/Editor/Blutility/Private/GlobalEditorUtilityBase.cpp
+++ b/Engine/Source/Editor/Blutility/Private/GlobalEditorUtilityBase.cpp
@@ -172,7 +172,7 @@ void UDEPRECATED_GlobalEditorUtilityBase::RenameAsset(UObject* Asset, const FStr
 AActor* UDEPRECATED_GlobalEditorUtilityBase::GetActorReference(FString PathToActor)
 {
 #if WITH_EDITOR
-	return Cast<AActor>(StaticFindObject(AActor::StaticClass(), GEditor->GetEditorWorldContext().World(), *PathToActor, false));
+	return Cast<AActor>(StaticFindObject(AActor::StaticClass(), GEditor->GetEditorWorldContext().World(), *PathToActor, EFindObjectFlags::None));
 #else
 	return nullptr;
 #endif //WITH_EDITOR
diff --git a/Engine/Source/Editor/Blutility/Private/PlacedEditorUtilityBase.cpp b/Engine/Source/Editor/Blutility/Private/PlacedEditorUtilityBase.cpp
--- a/Engine/Source/Editor/Blutility/Private/PlacedEditorUtilityBase.cpp
+++ b/Engine/Source/Editor/Blutility/Private/PlacedEditorUtilityBase.cpp
@@ -98,7 +98,7 @@ void ADEPRECATED_PlacedEditorUtilityBase::SetActorSelectionState(AActor* Actor,
 AActor* ADEPRECATED_PlacedEditorUtilityBase::GetActorReference(FString PathToActor)
 {
 #if WITH_EDITOR
-	return Cast<AActor>(StaticFindObject(AActor::StaticClass(), GEditor->GetEditorWorldContext().World(), *PathToActor, false));
+	return Cast<AActor>(StaticFindObject(AActor::StaticClass(), GEditor->GetEditorWorldContext().World(), *PathToActor, EFindObjectFlags::None));
 #else
 	return nullptr;
 #endif //WITH_EDITOR
diff --git a/Engine/Source/Editor/ContentBrowser/Private/SAssetDialog.cpp b/Engine/Source/Editor/ContentBrowser/Private/SAssetDialog.cpp
--- a/Engine/Source/Editor/ContentBrowser/Private/SAssetDialog.cpp
+++ b/Engine/Source/Editor/ContentBrowser/Private/SAssetDialog.cpp
@@ -1289,7 +1289,7 @@ void SAssetDialog::UpdateInputValidity()
 			const bool bAllowExistingAsset = (ExistingAssetPolicy == ESaveAssetDialogExistingAssetPolicy::AllowButWarn);
 
 			FTopLevelAssetPath AssetClassName = AssetClassNames.Num() == 1 ? AssetClassNames[0] : FTopLevelAssetPath();
-			UClass* AssetClass = !AssetClassName.IsNull() ? FindObject<UClass>(AssetClassName, true) : nullptr;
+			UClass* AssetClass = !AssetClassName.IsNull() ? FindObject<UClass>(AssetClassName, EFindObjectFlags::ExactClass) : nullptr;
 
 			if ( !ContentBrowserUtils::IsValidObjectPathForCreate(ObjectPath, AssetClass, ErrorMessage, bAllowExistingAsset) )
 			{
diff --git a/Engine/Source/Editor/DetailCustomizations/Private/MacTargetSettingsDetails.cpp b/Engine/Source/Editor/DetailCustomizations/Private/MacTargetSettingsDetails.cpp
--- a/Engine/Source/Editor/DetailCustomizations/Private/MacTargetSettingsDetails.cpp
+++ b/Engine/Source/Editor/DetailCustomizations/Private/MacTargetSettingsDetails.cpp
@@ -353,7 +353,7 @@ TSharedRef<SWidget> FMacTargetSettingsDetails::OnGetShaderVersionContent()
 {
 	FMenuBuilder MenuBuilder(true, NULL);
 	
-	UEnum* Enum = FindObjectChecked<UEnum>(nullptr, TEXT("/Script/MacTargetPlatformSettings.EMacMetalShaderStandard"), true);
+	UEnum* Enum = FindObjectChecked<UEnum>(nullptr, TEXT("/Script/MacTargetPlatformSettings.EMacMetalShaderStandard"), EFindObjectFlags::ExactClass);
 	
 	for (int32 i = 0; i < Enum->GetMaxEnumValue(); i++)
 	{
@@ -376,7 +376,7 @@ FText FMacTargetSettingsDetails::GetShaderVersionDesc() const
     int32 EnumValue;
     ShaderVersionPropertyHandle->GetValue(EnumValue);
 	
-	UEnum* Enum = FindObjectChecked<UEnum>(nullptr, TEXT("/Script/MacTargetPlatformSettings.EMacMetalShaderStandard"), true);
+	UEnum* Enum = FindObjectChecked<UEnum>(nullptr, TEXT("/Script/MacTargetPlatformSettings.EMacMetalShaderStandard"), EFindObjectFlags::ExactClass);
 	
 	if (EnumValue < Enum->GetMaxEnumValue() && Enum->IsValidEnumValue(EnumValue))
 	{
diff --git a/Engine/Source/Editor/Kismet/Private/BlueprintEditor.cpp b/Engine/Source/Editor/Kismet/Private/BlueprintEditor.cpp
--- a/Engine/Source/Editor/Kismet/Private/BlueprintEditor.cpp
+++ b/Engine/Source/Editor/Kismet/Private/BlueprintEditor.cpp
@@ -6827,7 +6827,7 @@ void FBlueprintEditor::ConvertEventToFunction(UK2Node_Event* SelectedEventNode)
 			const FString OriginalEventNameString = OriginalEventName.ToString();
 
 			// If there is already another object in the same scope with this name, rename it.
-			UObject* ExistingObject = StaticFindObject(/*Class=*/ nullptr, NodeBP, *OriginalEventNameString, true);
+			UObject* ExistingObject = StaticFindObject(/*Class=*/ nullptr, NodeBP, *OriginalEventNameString, EFindObjectFlags::ExactClass);
 			if (ExistingObject)
 			{
 				check(ExistingObject->GetOuter() == NewGraph->GetOuter());
diff --git a/Engine/Source/Editor/KismetCompiler/Private/KismetCompilerMisc.cpp b/Engine/Source/Editor/KismetCompiler/Private/KismetCompilerMisc.cpp
--- a/Engine/Source/Editor/KismetCompiler/Private/KismetCompilerMisc.cpp
+++ b/Engine/Source/Editor/KismetCompiler/Private/KismetCompilerMisc.cpp
@@ -1703,7 +1703,7 @@ FFieldVariant FKismetCompilerUtilities::CheckPropertyNameOnScope(UStruct* Scope,
 {
 	FString NameStr = PropertyName.ToString();
 
-	if (UObject* ExistingObject = FindObject<UObject>(Scope, *NameStr, false))
+	if (UObject* ExistingObject = FindObject<UObject>(Scope, *NameStr, EFindObjectFlags::None))
 	{
 		return ExistingObject;
 	}
diff --git a/Engine/Source/Editor/LevelEditor/Private/LightmapResRatioAdjust.cpp b/Engine/Source/Editor/LevelEditor/Private/LightmapResRatioAdjust.cpp
--- a/Engine/Source/Editor/LevelEditor/Private/LightmapResRatioAdjust.cpp
+++ b/Engine/Source/Editor/LevelEditor/Private/LightmapResRatioAdjust.cpp
@@ -107,10 +107,9 @@ bool FLightmapResRatioAdjustSettings::ApplyRatioAdjustment()
 				TArray<AActor*> AttachedActors;
 				Actor->GetAttachedActors( AttachedActors );
 
-				const bool bExactClass = true;
 				TArray<AActor*> AttachedBrushes;
 				// Get any brush actors attached to this one
-				if( ContainsObjectOfClass( AttachedActors, ABrush::StaticClass(), bExactClass, &AttachedBrushes ) )
+				if( ContainsObjectOfClass( AttachedActors, ABrush::StaticClass(), EFindObjectFlags::ExactClass, &AttachedBrushes ) )
 				{
 					for( int32 BrushIndex = 0; BrushIndex < AttachedBrushes.Num(); ++BrushIndex )
 					{
diff --git a/Engine/Source/Editor/MovieSceneTools/Private/MovieSceneToolHelpers.cpp b/Engine/Source/Editor/MovieSceneTools/Private/MovieSceneToolHelpers.cpp
--- a/Engine/Source/Editor/MovieSceneTools/Private/MovieSceneToolHelpers.cpp
+++ b/Engine/Source/Editor/MovieSceneTools/Private/MovieSceneToolHelpers.cpp
@@ -4379,7 +4379,7 @@ static FString GetSlateFromMovieSceneSequence(UMovieSceneSequence* MovieSceneSeq
 	}
 
 	// Find the UTakeMetaData class dynamically, since it is defined in a plugin.
-	static UClass* TakeMetaDataClass = FindObject<UClass>(nullptr, TEXT("/Script/TakesCore.TakeMetaData"), false);
+	static UClass* TakeMetaDataClass = FindObject<UClass>(nullptr, TEXT("/Script/TakesCore.TakeMetaData"), EFindObjectFlags::None);
 	if (!TakeMetaDataClass)
 	{
 		return LevelSequence->GetName();
diff --git a/Engine/Source/Editor/PropertyEditor/Private/EditConditionContext.cpp b/Engine/Source/Editor/PropertyEditor/Private/EditConditionContext.cpp
--- a/Engine/Source/Editor/PropertyEditor/Private/EditConditionContext.cpp
+++ b/Engine/Source/Editor/PropertyEditor/Private/EditConditionContext.cpp
@@ -82,7 +82,7 @@ const TWeakObjectPtr<UFunction> FEditConditionContext::GetFunction(const FString
 	if (Function == nullptr && FieldName.Contains(TEXT(".")))
 	{
 		// Function not found in struct, try to see if this is a static function
-		Function = FindObject<UFunction>(nullptr, *FieldName, true);
+		Function = FindObject<UFunction>(nullptr, *FieldName, EFindObjectFlags::ExactClass);
 
 		if (Function.IsValid() && !Function->HasAnyFunctionFlags(EFunctionFlags::FUNC_Static))
 		{
diff --git a/Engine/Source/Editor/PropertyEditor/Private/PropertyEditorUtils.cpp b/Engine/Source/Editor/PropertyEditor/Private/PropertyEditorUtils.cpp
--- a/Engine/Source/Editor/PropertyEditor/Private/PropertyEditorUtils.cpp
+++ b/Engine/Source/Editor/PropertyEditor/Private/PropertyEditorUtils.cpp
@@ -1,4 +1,4 @@
-﻿// Copyright Epic Games, Inc. All Rights Reserved.
+// Copyright Epic Games, Inc. All Rights Reserved.
 
 #include "PropertyEditorUtils.h"
 
@@ -49,7 +49,7 @@ namespace PropertyEditorUtils
 		if (InOutPropertyPath.Contains(TEXT(".")))
 		{
 			InOutContainers.Empty();
-			UFunction* GetOptionsFunction = FindObject<UFunction>(nullptr, *InOutPropertyPath, true);
+			UFunction* GetOptionsFunction = FindObject<UFunction>(nullptr, *InOutPropertyPath, EFindObjectFlags::ExactClass);
 
 			if (ensureMsgf(GetOptionsFunction && GetOptionsFunction->HasAnyFunctionFlags(EFunctionFlags::FUNC_Static), TEXT("Invalid GetOptions: %s"), *InOutPropertyPath))
 			{
@@ -213,7 +213,7 @@ namespace PropertyEditorUtils
 				// First look for a library function assuming fully qualified path, e.g. /Script/ModuleName.ClassName:FunctionName
 				if(GetAllowedClassesFunctionName.Contains(TEXT(".")))
 				{
-					const UFunction* GetAllowedClassesFunction = FindObject<UFunction>(nullptr, *GetAllowedClassesFunctionName, true);
+					const UFunction* GetAllowedClassesFunction = FindObject<UFunction>(nullptr, *GetAllowedClassesFunctionName, EFindObjectFlags::ExactClass);
 					if (ensureMsgf(GetAllowedClassesFunction && GetAllowedClassesFunction->HasAnyFunctionFlags(EFunctionFlags::FUNC_Static), TEXT("Invalid GetAllowedClasses: %s"), *GetAllowedClassesFunctionName))
 					{
 						UObject* GetAllowedClassesCDO = GetAllowedClassesFunction->GetOwnerClass()->GetDefaultObject();
diff --git a/Engine/Source/Editor/Sequencer/Private/SSequencerSection.cpp b/Engine/Source/Editor/Sequencer/Private/SSequencerSection.cpp
--- a/Engine/Source/Editor/Sequencer/Private/SSequencerSection.cpp
+++ b/Engine/Source/Editor/Sequencer/Private/SSequencerSection.cpp
@@ -383,7 +383,7 @@ struct FSequencerSectionPainterImpl : FSequencerSectionPainter
 
 		TSharedRef<FSlateFontCache> FontCache = FSlateApplication::Get().GetRenderer()->GetFontCache();
 
-		UEnum* Enum = FindObjectChecked<UEnum>(nullptr, TEXT("/Script/MovieScene.EMovieSceneBlendType"), true);
+		UEnum* Enum = FindObjectChecked<UEnum>(nullptr, TEXT("/Script/MovieScene.EMovieSceneBlendType"), EFindObjectFlags::ExactClass);
 		FText DisplayText = Enum->GetDisplayNameTextByValue((int64)SectionObject->GetBlendType().Get());
 
 		FSlateFontInfo FontInfo = FAppStyle::GetFontStyle("Sequencer.Section.BackgroundText");
diff --git a/Engine/Source/Editor/Sequencer/Public/Filters/SequencerFilterBase.h b/Engine/Source/Editor/Sequencer/Public/Filters/SequencerFilterBase.h
--- a/Engine/Source/Editor/Sequencer/Public/Filters/SequencerFilterBase.h
+++ b/Engine/Source/Editor/Sequencer/Public/Filters/SequencerFilterBase.h
@@ -1,4 +1,4 @@
-﻿// Copyright Epic Games, Inc. All Rights Reserved.
+// Copyright Epic Games, Inc. All Rights Reserved.
 
 #pragma once
 
@@ -16,7 +16,7 @@ class FSequencerFilterBase
 public:
 	static bool SupportsLevelSequence(UMovieSceneSequence* const InSequence)
 	{
-		static UClass* const LevelSequenceClass = FindObject<UClass>(nullptr, TEXT("/Script/LevelSequence.LevelSequence"), true);
+		static UClass* const LevelSequenceClass = FindObject<UClass>(nullptr, TEXT("/Script/LevelSequence.LevelSequence"), EFindObjectFlags::ExactClass);
 		return InSequence
 			&& LevelSequenceClass
 			&& InSequence->GetClass()->IsChildOf(LevelSequenceClass);
@@ -24,7 +24,7 @@ public:
 
 	static bool SupportsUMGSequence(UMovieSceneSequence* const InSequence)
 	{
-		static UClass* const WidgetAnimationClass = FindObject<UClass>(nullptr, TEXT("/Script/UMG.WidgetAnimation"), true);
+		static UClass* const WidgetAnimationClass = FindObject<UClass>(nullptr, TEXT("/Script/UMG.WidgetAnimation"), EFindObjectFlags::ExactClass);
 		return InSequence
 			&& WidgetAnimationClass
 			&& InSequence->GetClass()->IsChildOf(WidgetAnimationClass);
@@ -32,7 +32,7 @@ public:
 
 	static bool SupportsDaySequence(UMovieSceneSequence* const InSequence)
 	{
-		static UClass* const DaySequenceClass = FindObject<UClass>(nullptr, TEXT("/Script/DaySequence.DaySequence"), true);
+		static UClass* const DaySequenceClass = FindObject<UClass>(nullptr, TEXT("/Script/DaySequence.DaySequence"), EFindObjectFlags::ExactClass);
 		return InSequence
 			&& DaySequenceClass
 			&& InSequence->GetClass()->IsChildOf(DaySequenceClass);
diff --git a/Engine/Source/Editor/StructUtilsEditor/Private/Customizations/PropertyBagDetails.cpp b/Engine/Source/Editor/StructUtilsEditor/Private/Customizations/PropertyBagDetails.cpp
--- a/Engine/Source/Editor/StructUtilsEditor/Private/Customizations/PropertyBagDetails.cpp
+++ b/Engine/Source/Editor/StructUtilsEditor/Private/Customizations/PropertyBagDetails.cpp
@@ -493,7 +493,7 @@ bool FindUserFunction(TSharedPtr<IPropertyHandle> InStructProperty, FName InFunc
 	// Check for external function references, taken from GetOptions
 	if (FunctionName.Contains(TEXT(".")))
 	{
-		OutFunc = FindObject<UFunction>(nullptr, *FunctionName, true);
+		OutFunc = FindObject<UFunction>(nullptr, *FunctionName, EFindObjectFlags::ExactClass);
 
 		if (ensureMsgf(OutFunc && OutFunc->HasAnyFunctionFlags(EFunctionFlags::FUNC_Static), TEXT("[%s] Didn't find function %s or expected it to be static"), *InFuncMetadataName.ToString(), *FunctionName))
 		{
diff --git a/Engine/Source/Editor/UMGEditor/Private/WidgetBlueprintEditorUtils.cpp b/Engine/Source/Editor/UMGEditor/Private/WidgetBlueprintEditorUtils.cpp
--- a/Engine/Source/Editor/UMGEditor/Private/WidgetBlueprintEditorUtils.cpp
+++ b/Engine/Source/Editor/UMGEditor/Private/WidgetBlueprintEditorUtils.cpp
@@ -3083,7 +3083,7 @@ TTuple<float, FVector2D> FWidgetBlueprintEditorUtils::GetThumbnailImageScaleAndO
 void FWidgetBlueprintEditorUtils::SetTextureAsAssetThumbnail(UWidgetBlueprint* WidgetBlueprint, UTexture2D* ThumbnailTexture)
 {
 	const TCHAR* ThumbnailName = TEXT("Thumbnail");
-	UTexture2D* ExistingThumbnail = FindObject<UTexture2D>(WidgetBlueprint, ThumbnailName, false);
+	UTexture2D* ExistingThumbnail = FindObject<UTexture2D>(WidgetBlueprint, ThumbnailName, EFindObjectFlags::None);
 	if (ExistingThumbnail)
 	{
 		ExistingThumbnail->Rename(nullptr, GetTransientPackage());
diff --git a/Engine/Source/Editor/UnrealEd/Private/Commandlets/PackageUtilities.cpp b/Engine/Source/Editor/UnrealEd/Private/Commandlets/PackageUtilities.cpp
--- a/Engine/Source/Editor/UnrealEd/Private/Commandlets/PackageUtilities.cpp
+++ b/Engine/Source/Editor/UnrealEd/Private/Commandlets/PackageUtilities.cpp
@@ -252,7 +252,7 @@ bool NormalizePackageNames( TArray<FString> PackageNames, TArray<FString>& Packa
 			{
 				PackageName = PackagePathNames[PackageIndex];
 			}
-			UPackage* ExistingPackage = FindObject<UPackage>(NULL, *PackageName, true);
+			UPackage* ExistingPackage = FindObject<UPackage>(NULL, *PackageName, EFindObjectFlags::ExactClass);
 			if ( ExistingPackage != NULL )
 			{
 				// skip resetting loaders on default materials since they are expected to be post-loaded at that point
@@ -788,7 +788,7 @@ int32 ULoadPackageCommandlet::Main( const FString& Params )
 		FString PackageName = PackagePath.GetPackageName();
 		if (!PackageName.IsEmpty())
 		{
-			UPackage* Package = FindObject<UPackage>(nullptr, *PackageName, true);
+			UPackage* Package = FindObject<UPackage>(nullptr, *PackageName, EFindObjectFlags::ExactClass);
 			if (Package != NULL && !bLoadAllPackages)
 			{
 				ResetLoaders(Package);
@@ -1601,7 +1601,7 @@ int32 UPkgInfoCommandlet::Main( const FString& Params )
 			FString PackageName;
 			if (FPackageName::TryConvertFilenameToLongPackageName(Filename, PackageName))
 			{
-				UPackage* ExistingPackage = FindObject<UPackage>(nullptr, *PackageName, true);
+				UPackage* ExistingPackage = FindObject<UPackage>(nullptr, *PackageName, EFindObjectFlags::ExactClass);
 				if (ExistingPackage != nullptr)
 				{
 					ResetLoaders(ExistingPackage);
diff --git a/Engine/Source/Editor/UnrealEd/Private/Editor.cpp b/Engine/Source/Editor/UnrealEd/Private/Editor.cpp
--- a/Engine/Source/Editor/UnrealEd/Private/Editor.cpp
+++ b/Engine/Source/Editor/UnrealEd/Private/Editor.cpp
@@ -1148,8 +1148,7 @@ namespace EditorUtilities
 					if( Level->GetFName() == SimWorldActor->GetLevel()->GetFName() )
 					{
 						// Find our counterpart actor
-						const bool bExactClass = false;	// Don't match class exactly, because we support all classes derived from Actor as well!
-						AActor* EditorWorldActor = FindObject<AActor>( Level, *SimWorldActor->GetFName().ToString(), bExactClass );
+						AActor* EditorWorldActor = FindObject<AActor>( Level, *SimWorldActor->GetFName().ToString(), EFindObjectFlags::None ); // Don't match class exactly, because we support all classes derived from Actor as well!
 						if( EditorWorldActor )
 						{
 							return EditorWorldActor;
@@ -1180,8 +1179,7 @@ namespace EditorUtilities
 					if (Level->GetFName() == EditorWorldActor->GetLevel()->GetFName())
 					{
 						// Find our counterpart actor
-						const bool bExactClass = false;	// Don't match class exactly, because we support all classes derived from Actor as well!
-						AActor* SimWorldActor = FindObject<AActor>(Level, *EditorWorldActor->GetFName().ToString(), bExactClass);
+						AActor* SimWorldActor = FindObject<AActor>(Level, *EditorWorldActor->GetFName().ToString(), EFindObjectFlags::None); // Don't match class exactly, because we support all classes derived from Actor as well!
 						if (SimWorldActor && GEditor->ObjectsThatExistInEditorWorld.Get(SimWorldActor))
 						{
 							return SimWorldActor;
diff --git a/Engine/Source/Editor/UnrealEd/Private/EditorBuildUtils.cpp b/Engine/Source/Editor/UnrealEd/Private/EditorBuildUtils.cpp
--- a/Engine/Source/Editor/UnrealEd/Private/EditorBuildUtils.cpp
+++ b/Engine/Source/Editor/UnrealEd/Private/EditorBuildUtils.cpp
@@ -1415,7 +1415,7 @@ void FEditorBuildUtils::TriggerHierarchicalLODBuilder(UWorld* InWorld)
 {
 	if (InWorld->IsPartitionedWorld())
 	{
-		TSubclassOf<UWorldPartitionBuilder> WorldPartitionHLODsBuilder = FindObjectChecked<UClass>(nullptr, TEXT("/Script/UnrealEd.WorldPartitionHLODsBuilder"), true);
+		TSubclassOf<UWorldPartitionBuilder> WorldPartitionHLODsBuilder = FindObjectChecked<UClass>(nullptr, TEXT("/Script/UnrealEd.WorldPartitionHLODsBuilder"), EFindObjectFlags::ExactClass);
 		IWorldPartitionEditorModule::Get().RunBuilder(WorldPartitionHLODsBuilder, InWorld);
 	}
 	else
@@ -1429,7 +1429,7 @@ void FEditorBuildUtils::TriggerMinimapBuilder(UWorld* InWorld)
 {
 	if (InWorld->IsPartitionedWorld())
 	{
-		TSubclassOf<UWorldPartitionBuilder> WorldPartitionMiniMapBuilder = FindObjectChecked<UClass>(nullptr, TEXT("/Script/UnrealEd.WorldPartitionMiniMapBuilder"), true);
+		TSubclassOf<UWorldPartitionBuilder> WorldPartitionMiniMapBuilder = FindObjectChecked<UClass>(nullptr, TEXT("/Script/UnrealEd.WorldPartitionMiniMapBuilder"), EFindObjectFlags::ExactClass);
 		IWorldPartitionEditorModule::Get().RunBuilder(WorldPartitionMiniMapBuilder, InWorld);
 	}
 }
@@ -1438,7 +1438,7 @@ void FEditorBuildUtils::TriggerLandscapeSplineMeshesBuilder(UWorld* InWorld)
 {
 	if (InWorld->IsPartitionedWorld())
 	{
-		TSubclassOf<UWorldPartitionBuilder> WorldPartitionLandscapeSplineMeshesBuilder = FindObjectChecked<UClass>(nullptr, TEXT("/Script/UnrealEd.WorldPartitionLandscapeSplineMeshesBuilder"), true);
+		TSubclassOf<UWorldPartitionBuilder> WorldPartitionLandscapeSplineMeshesBuilder = FindObjectChecked<UClass>(nullptr, TEXT("/Script/UnrealEd.WorldPartitionLandscapeSplineMeshesBuilder"), EFindObjectFlags::ExactClass);
 		IWorldPartitionEditorModule::Get().RunBuilder(WorldPartitionLandscapeSplineMeshesBuilder, InWorld);
 	}
 }
diff --git a/Engine/Source/Editor/UnrealEd/Private/EditorEngine.cpp b/Engine/Source/Editor/UnrealEd/Private/EditorEngine.cpp
--- a/Engine/Source/Editor/UnrealEd/Private/EditorEngine.cpp
+++ b/Engine/Source/Editor/UnrealEd/Private/EditorEngine.cpp
@@ -3938,7 +3938,7 @@ UBrushBuilder* UEditorEngine::FindBrushBuilder( UClass* BrushBuilderClass )
 	TArray< UBrushBuilder* > FoundBuilders;
 	UBrushBuilder* Builder = NULL;
 	// Search for the existing brush builder
-	if( ContainsObjectOfClass<UBrushBuilder>( BrushBuilders, BrushBuilderClass, true, &FoundBuilders ) )
+	if( ContainsObjectOfClass<UBrushBuilder>( BrushBuilders, BrushBuilderClass, EFindObjectFlags::ExactClass, &FoundBuilders ) )
 	{
 		// Should not be more than one of the same type
 		check( FoundBuilders.Num() == 1 );
diff --git a/Engine/Source/Editor/UnrealEd/Private/EditorServer.cpp b/Engine/Source/Editor/UnrealEd/Private/EditorServer.cpp
--- a/Engine/Source/Editor/UnrealEd/Private/EditorServer.cpp
+++ b/Engine/Source/Editor/UnrealEd/Private/EditorServer.cpp
@@ -1831,10 +1831,9 @@ void UEditorEngine::RebuildAlteredBSP()
 				TArray<AActor*> AttachedActors;
 				Actor->GetAttachedActors(AttachedActors);
 
-				const bool bExactClass = true;
 				TArray<AActor*> AttachedBrushes;
 				// Get any brush actors attached to the selected actor
-				if (ContainsObjectOfClass(AttachedActors, ABrush::StaticClass(), bExactClass, &AttachedBrushes))
+				if (ContainsObjectOfClass(AttachedActors, ABrush::StaticClass(), EFindObjectFlags::ExactClass, &AttachedBrushes))
 				{
 					for (int32 BrushIndex = 0; BrushIndex < AttachedBrushes.Num(); ++BrushIndex)
 					{
@@ -3211,7 +3210,7 @@ ULevel*  UEditorEngine::CreateTransLevelMoveBuffer( UWorld* InWorld )
 	ABrush* BufferDefaultBrush = InWorld->SpawnActor<ABrush>( SpawnInfo );
 
 	check( BufferDefaultBrush->GetBrushComponent() );
-	BufferDefaultBrush->Brush = CastChecked<UModel>(StaticFindObject(UModel::StaticClass(), BufferLevel->OwningWorld->GetOuter(), TEXT("Brush"), true), ECastCheckedType::NullAllowed);
+	BufferDefaultBrush->Brush = CastChecked<UModel>(StaticFindObject(UModel::StaticClass(), BufferLevel->OwningWorld->GetOuter(), TEXT("Brush"), EFindObjectFlags::ExactClass), ECastCheckedType::NullAllowed);
 	if (!BufferDefaultBrush->Brush)
 	{
 		BufferDefaultBrush->Brush = NewObject<UModel>(InWorld, TEXT("Brush"));
@@ -5576,7 +5575,7 @@ void ListMapPackageDependencies(const TCHAR* InStr)
 					for (int32 ExportIdx = 0; ExportIdx < Linker->ExportMap.Num(); ExportIdx++)
 					{
 						FTopLevelAssetPath CheckClassPathName(Linker->GetExportClassPackage(ExportIdx), Linker->GetExportClassName(ExportIdx));
-						UClass* CheckClass = FindObject<UClass>(CheckClassPathName, true);
+						UClass* CheckClass = FindObject<UClass>(CheckClassPathName, EFindObjectFlags::ExactClass);
 						if (
 							(CheckClass != NULL) &&
 							(CheckClass->IsChildOf(UTexture::StaticClass()) == true)
diff --git a/Engine/Source/Editor/UnrealEd/Private/Factories/MaterialImportHelpers.cpp b/Engine/Source/Editor/UnrealEd/Private/Factories/MaterialImportHelpers.cpp
--- a/Engine/Source/Editor/UnrealEd/Private/Factories/MaterialImportHelpers.cpp
+++ b/Engine/Source/Editor/UnrealEd/Private/Factories/MaterialImportHelpers.cpp
@@ -15,8 +15,7 @@ UMaterialInterface* UMaterialImportHelpers::FindExistingMaterialFromSearchLocati
 	}
 
 	//Search in memory
-	constexpr bool bExactClass = false;
-	UMaterialInterface* FoundMaterial = FindObject<UMaterialInterface>(nullptr, *MaterialFullName, bExactClass);
+	UMaterialInterface* FoundMaterial = FindObject<UMaterialInterface>(nullptr, *MaterialFullName, EFindObjectFlags::None);
 
 	if (FoundMaterial == nullptr)
 	{
diff --git a/Engine/Source/Editor/UnrealEd/Private/Fbx/FbxSkeletalMeshImport.cpp b/Engine/Source/Editor/UnrealEd/Private/Fbx/FbxSkeletalMeshImport.cpp
--- a/Engine/Source/Editor/UnrealEd/Private/Fbx/FbxSkeletalMeshImport.cpp
+++ b/Engine/Source/Editor/UnrealEd/Private/Fbx/FbxSkeletalMeshImport.cpp
@@ -1714,7 +1714,7 @@ USkeletalMesh* UnFbx::FFbxImporter::ImportSkeletalMesh(FImportSkeletalMeshArgs &
 	USkeletalMesh* ExistingSkelMesh = nullptr;
 	if ( !ImportSkeletalMeshArgs.FbxShapeArray  )
 	{
-		UObject* ExistingObject = StaticFindObjectFast(UObject::StaticClass(), ImportSkeletalMeshArgs.InParent, ImportSkeletalMeshArgs.Name, false, RF_NoFlags, EInternalObjectFlags::Garbage);
+		UObject* ExistingObject = StaticFindObjectFast(UObject::StaticClass(), ImportSkeletalMeshArgs.InParent, ImportSkeletalMeshArgs.Name, EFindObjectFlags::None, RF_NoFlags, EInternalObjectFlags::Garbage);
 		ExistingSkelMesh = Cast<USkeletalMesh>(ExistingObject);
 
 		if (!ExistingSkelMesh && ExistingObject)
diff --git a/Engine/Source/Editor/UnrealEd/Private/Fbx/FbxStaticMeshImport.cpp b/Engine/Source/Editor/UnrealEd/Private/Fbx/FbxStaticMeshImport.cpp
--- a/Engine/Source/Editor/UnrealEd/Private/Fbx/FbxStaticMeshImport.cpp
+++ b/Engine/Source/Editor/UnrealEd/Private/Fbx/FbxStaticMeshImport.cpp
@@ -102,7 +102,7 @@ struct FRestoreReimportData
 		//Restore the duplicate asset
 		OriginalPackageName = UPackageTools::SanitizePackageName(OriginalPackageName);
 		Package = CreatePackage( *OriginalPackageName);
-		UObject* ExistingObject = StaticFindObject(UStaticMesh::StaticClass(), Package, *OriginalName, true);
+		UObject* ExistingObject = StaticFindObject(UStaticMesh::StaticClass(), Package, *OriginalName, EFindObjectFlags::ExactClass);
 		if (ExistingObject)
 		{
 			//Some reimport path did not trash the package in case there is a fail (i.e. scene reimport)
diff --git a/Engine/Source/Editor/UnrealEd/Private/Kismet2/BlueprintEditorUtils.cpp b/Engine/Source/Editor/UnrealEd/Private/Kismet2/BlueprintEditorUtils.cpp
--- a/Engine/Source/Editor/UnrealEd/Private/Kismet2/BlueprintEditorUtils.cpp
+++ b/Engine/Source/Editor/UnrealEd/Private/Kismet2/BlueprintEditorUtils.cpp
@@ -8065,7 +8065,7 @@ bool FBlueprintEditorUtils::RenameTimeline(UBlueprint* Blueprint, const FName Ol
 
 			Blueprint->Timelines.Remove(Template);
 			
-			UObject* ExistingObject = StaticFindObject(nullptr, Template->GetOuter(), *NewTemplateName, true);
+			UObject* ExistingObject = StaticFindObject(nullptr, Template->GetOuter(), *NewTemplateName, EFindObjectFlags::ExactClass);
 			if (ExistingObject != Template && ExistingObject != nullptr)
 			{
 				ExistingObject->Rename(*MakeUniqueObjectName(ExistingObject->GetOuter(), ExistingObject->GetClass(), ExistingObject->GetFName()).ToString());
diff --git a/Engine/Source/Editor/UnrealEd/Private/Kismet2/Kismet2NameValidators.cpp b/Engine/Source/Editor/UnrealEd/Private/Kismet2/Kismet2NameValidators.cpp
--- a/Engine/Source/Editor/UnrealEd/Private/Kismet2/Kismet2NameValidators.cpp
+++ b/Engine/Source/Editor/UnrealEd/Private/Kismet2/Kismet2NameValidators.cpp
@@ -155,7 +155,7 @@ EValidatorResult FKismetNameValidator::IsValid(const FName& Name, bool /* bOrigi
 			ValidatorResult = EValidatorResult::Ok;
 
 			// Check for collision with an existing object.
-			if (UObject* ExistingObject = StaticFindObject(/*Class=*/ nullptr, const_cast<UBlueprint*>(BlueprintObject), *Name.ToString(), true))
+			if (UObject* ExistingObject = StaticFindObject(/*Class=*/ nullptr, const_cast<UBlueprint*>(BlueprintObject), *Name.ToString(), EFindObjectFlags::ExactClass))
 			{
 				// To allow the linker to resolve imports when dependent Blueprints are loaded, macro libraries
 				// will leave behind a redirector whenever one of its macro graph objects are renamed. These get
diff --git a/Engine/Source/Editor/UnrealEd/Private/ObjectTools.cpp b/Engine/Source/Editor/UnrealEd/Private/ObjectTools.cpp
--- a/Engine/Source/Editor/UnrealEd/Private/ObjectTools.cpp
+++ b/Engine/Source/Editor/UnrealEd/Private/ObjectTools.cpp
@@ -2193,7 +2193,7 @@ namespace ObjectTools
 				// Construct the ignore package list.
 				for( int32 PackageNameItr = 0; PackageNameItr < IgnorePackageNames.Num(); ++PackageNameItr )
 				{
-					UPackage* PackageToIgnore = FindObject<UPackage>(NULL,*(IgnorePackageNames[PackageNameItr]),true);
+					UPackage* PackageToIgnore = FindObject<UPackage>(NULL,*(IgnorePackageNames[PackageNameItr]),EFindObjectFlags::ExactClass);
 
 					if( PackageToIgnore == NULL )
 					{// An invalid package name was provided.
diff --git a/Engine/Source/Editor/UnrealEd/Private/ReferenceInfoUtils.cpp b/Engine/Source/Editor/UnrealEd/Private/ReferenceInfoUtils.cpp
--- a/Engine/Source/Editor/UnrealEd/Private/ReferenceInfoUtils.cpp
+++ b/Engine/Source/Editor/UnrealEd/Private/ReferenceInfoUtils.cpp
@@ -666,7 +666,7 @@ namespace ReferenceInfoUtils
 
 		for (int32 AssetIdx = 0; AssetIdx < AssetData.Num(); ++AssetIdx)
 		{
-			IgnorePackages.Add(FindObject<UPackage>(NULL, *AssetData[AssetIdx].PackageName.ToString(), true));
+			IgnorePackages.Add(FindObject<UPackage>(NULL, *AssetData[AssetIdx].PackageName.ToString(), EFindObjectFlags::ExactClass));
 		}
 
 		IgnorePackages.Add(GetTransientPackage());
diff --git a/Engine/Source/Editor/UnrealEd/Private/ReferencedAssetsUtils.cpp b/Engine/Source/Editor/UnrealEd/Private/ReferencedAssetsUtils.cpp
--- a/Engine/Source/Editor/UnrealEd/Private/ReferencedAssetsUtils.cpp
+++ b/Engine/Source/Editor/UnrealEd/Private/ReferencedAssetsUtils.cpp
@@ -40,7 +40,7 @@ FFindReferencedAssets::FFindReferencedAssets(void)
 
 	for (int32 AssetIdx = 0; AssetIdx < AssetData.Num(); ++AssetIdx)
 	{
-		UPackage* Package = FindObject<UPackage>(NULL, *AssetData[AssetIdx].PackageName.ToString(), true);
+		UPackage* Package = FindObject<UPackage>(NULL, *AssetData[AssetIdx].PackageName.ToString(), EFindObjectFlags::ExactClass);
 		if( Package != NULL )
 		{
 			IgnorePackages.Add(Package);
diff --git a/Engine/Source/Editor/UnrealEd/Private/Subsystems/CollectionManagerScriptingSubsystem.cpp b/Engine/Source/Editor/UnrealEd/Private/Subsystems/CollectionManagerScriptingSubsystem.cpp
--- a/Engine/Source/Editor/UnrealEd/Private/Subsystems/CollectionManagerScriptingSubsystem.cpp
+++ b/Engine/Source/Editor/UnrealEd/Private/Subsystems/CollectionManagerScriptingSubsystem.cpp
@@ -25,7 +25,7 @@ void LogLastCollectionManagerError(const FText& Error)
 
 FString ECollectionScriptingShareType_To_String(const ECollectionScriptingShareType ShareType)
 {
-	static const UEnum* Enum = FindObject<UEnum>(nullptr, TEXT("/Script/CollectionManager.ECollectionScriptingShareType"), true);
+	static const UEnum* Enum = FindObject<UEnum>(nullptr, TEXT("/Script/CollectionManager.ECollectionScriptingShareType"), EFindObjectFlags::ExactClass);
 	if (Enum)
 	{
 		return Enum->GetDisplayNameTextByValue((uint64)ShareType).ToString();
diff --git a/Engine/Source/Editor/UnrealEd/Private/Subsystems/EditorActorSubsystem.cpp b/Engine/Source/Editor/UnrealEd/Private/Subsystems/EditorActorSubsystem.cpp
--- a/Engine/Source/Editor/UnrealEd/Private/Subsystems/EditorActorSubsystem.cpp
+++ b/Engine/Source/Editor/UnrealEd/Private/Subsystems/EditorActorSubsystem.cpp
@@ -494,7 +494,7 @@ void UEditorActorSubsystem::SetActorSelectionState(AActor* Actor, bool bShouldBe
 
 AActor* UEditorActorSubsystem::GetActorReference(FString PathToActor)
 {
-	return Cast<AActor>(StaticFindObject(AActor::StaticClass(), GEditor->GetEditorWorldContext().World(), *PathToActor, false));
+	return Cast<AActor>(StaticFindObject(AActor::StaticClass(), GEditor->GetEditorWorldContext().World(), *PathToActor, EFindObjectFlags::None));
 }
 
 
diff --git a/Engine/Source/Editor/UnrealEd/Public/AnimationEditorUtils.h b/Engine/Source/Editor/UnrealEd/Public/AnimationEditorUtils.h
--- a/Engine/Source/Editor/UnrealEd/Public/AnimationEditorUtils.h
+++ b/Engine/Source/Editor/UnrealEd/Public/AnimationEditorUtils.h
@@ -245,7 +245,7 @@ namespace AnimationEditorUtils
 					// @TODO: this doesn't work
 					//FString LongPackagePath = FPackageName::GetLongPackagePath(PackageName);
 					UObject* 	Parent = FindPackage(NULL, *PackageName);
-					UObject* NewAsset = FindObject<UObject>(Parent, *Name, false);
+					UObject* NewAsset = FindObject<UObject>(Parent, *Name, EFindObjectFlags::None);
 					if(NewAsset)
 					{
 						TArray<UObject*> NewAssets;
diff --git a/Engine/Source/Runtime/AssetRegistry/Private/AssetRegistry.cpp b/Engine/Source/Runtime/AssetRegistry/Private/AssetRegistry.cpp
--- a/Engine/Source/Runtime/AssetRegistry/Private/AssetRegistry.cpp
+++ b/Engine/Source/Runtime/AssetRegistry/Private/AssetRegistry.cpp
@@ -6667,7 +6667,7 @@ bool FAssetRegistryImpl::TryPostLoadAssetRegistryTags(FAssetData* AssetData)
 	UClass* AssetClass = nullptr;
 	FTopLevelAssetPath AssetClassPath;
 	AssetClassPath = AssetData->AssetClassPath;
-	AssetClass = FindObject<UClass>(AssetClassPath, true);
+	AssetClass = FindObject<UClass>(AssetClassPath, EFindObjectFlags::ExactClass);
 
 	while (!AssetClass)
 	{
@@ -6676,7 +6676,7 @@ bool FAssetRegistryImpl::TryPostLoadAssetRegistryTags(FAssetData* AssetData)
 		if (ParentClassPath && !ParentClassPath->IsNull())
 		{
 			AssetClassPath = *ParentClassPath;
-			AssetClass = FindObject<UClass>(AssetClassPath, true);
+			AssetClass = FindObject<UClass>(AssetClassPath, EFindObjectFlags::ExactClass);
 		}
 		else
 		{
@@ -6698,7 +6698,7 @@ bool FAssetRegistryImpl::TryPostLoadAssetRegistryTags(FAssetData* AssetData)
 
 			if (AssetClassPath != LastAssetClassPath && !AssetClassPath.IsNull())
 			{
-				AssetClass = FindObject<UClass>(AssetClassPath, true);
+				AssetClass = FindObject<UClass>(AssetClassPath, EFindObjectFlags::ExactClass);
 			}
 			else
 			{
diff --git a/Engine/Source/Runtime/ClothingSystemRuntimeCommon/Private/ClothingAsset.cpp b/Engine/Source/Runtime/ClothingSystemRuntimeCommon/Private/ClothingAsset.cpp
--- a/Engine/Source/Runtime/ClothingSystemRuntimeCommon/Private/ClothingAsset.cpp
+++ b/Engine/Source/Runtime/ClothingSystemRuntimeCommon/Private/ClothingAsset.cpp
@@ -1264,7 +1264,7 @@ bool UClothingAssetCommon::AddClothConfigs()
 				if (!ClothConfigPtr || !*ClothConfigPtr)
 				{
 					// Create new config object
-					checkf(!StaticFindObject(ClothConfigClass, this, *ClothConfigClass->GetName(), true), TEXT("Found an existing instance of ClothConfigClass in %s. Class:%s"), *GetPathNameSafe(this), *ClothConfigClass->GetName());
+					checkf(!StaticFindObject(ClothConfigClass, this, *ClothConfigClass->GetName(), EFindObjectFlags::ExactClass), TEXT("Found an existing instance of ClothConfigClass in %s. Class:%s"), *GetPathNameSafe(this), *ClothConfigClass->GetName());
 					UClothConfigBase* const ClothConfig = NewObject<UClothConfigBase>(this, ClothConfigClass, ClothConfigClass->GetFName(), RF_Transactional);
 
 					// Use the legacy config struct to try find a common config as an acceptable migration source
diff --git a/Engine/Source/Runtime/CoreUObject/Private/UObject/LinkerLoad.cpp b/Engine/Source/Runtime/CoreUObject/Private/UObject/LinkerLoad.cpp
--- a/Engine/Source/Runtime/CoreUObject/Private/UObject/LinkerLoad.cpp
+++ b/Engine/Source/Runtime/CoreUObject/Private/UObject/LinkerLoad.cpp
@@ -4322,9 +4322,9 @@ UClass* FLinkerLoad::TryCreatePlaceholderClassImport(int32 ImportIndex)
 		return nullptr;
 	}
 
-	if (UObject* ImportClassPackage = FindObjectFast<UPackage>(nullptr, Import.ClassPackage, /*bExactClass =*/ true))
+	if (UObject* ImportClassPackage = FindObjectFast<UPackage>(nullptr, Import.ClassPackage, EFindObjectFlags::ExactClass))
 	{
-		if (UClass* ImportClass = FindObjectFast<UClass>(ImportClassPackage, Import.ClassName, /*bExactClass =*/ false))
+		if (UClass* ImportClass = FindObjectFast<UClass>(ImportClassPackage, Import.ClassName, EFindObjectFlags::None))
 		{
 			if (UE::CanCreatePropertyBagPlaceholderTypeForImportClass(ImportClass))
 			{
@@ -5472,7 +5472,7 @@ UObject* FLinkerLoad::CreateExport( int32 Index )
 #if WITH_EDITOR
 		if ( GIsEditor && GIsRunning && !Export.Object )
 		{
-			UObjectRedirector* Redirector = (UObjectRedirector*)StaticFindObjectFast(UObjectRedirector::StaticClass(), ThisParent, Export.ObjectName, /*bExactClass*/true);
+			UObjectRedirector* Redirector = (UObjectRedirector*)StaticFindObjectFast(UObjectRedirector::StaticClass(), ThisParent, Export.ObjectName, EFindObjectFlags::ExactClass);
 			if (Redirector && Redirector->DestinationObject && Redirector->DestinationObject->IsA(LoadClass))
 			{
 				// A redirector has been found, replace this export with it.
diff --git a/Engine/Source/Runtime/CoreUObject/Private/UObject/UObjectConstructInternal.h b/Engine/Source/Runtime/CoreUObject/Private/UObject/UObjectConstructInternal.h
--- a/Engine/Source/Runtime/CoreUObject/Private/UObject/UObjectConstructInternal.h
+++ b/Engine/Source/Runtime/CoreUObject/Private/UObject/UObjectConstructInternal.h
@@ -284,7 +284,7 @@ void ConstructUFunctionHelper(UFunction*& InOutFunction, const FunctionParams& P
 	// is created.  Only do this if we reuse the current instance.
 	UField* PrevFunctionNextField = nullptr;
 	UFunction* PrevFunction = nullptr;
-	if (UObject* PrevObject = StaticFindObjectFastInternal( /*Class=*/ nullptr, Outer, FuncName, true))
+	if (UObject* PrevObject = StaticFindObjectFastInternal( /*Class=*/ nullptr, Outer, FuncName, EFindObjectFlags::ExactClass))
 	{
 		PrevFunction = Cast<UFunction>(PrevObject);
 		if (PrevFunction != nullptr)
diff --git a/Engine/Source/Runtime/CoreUObject/Public/UObject/StructOnScope.h b/Engine/Source/Runtime/CoreUObject/Public/UObject/StructOnScope.h
--- a/Engine/Source/Runtime/CoreUObject/Public/UObject/StructOnScope.h
+++ b/Engine/Source/Runtime/CoreUObject/Public/UObject/StructOnScope.h
@@ -325,7 +325,7 @@ public:
 			Ar << StructPath;
 			if (!StructPath.IsEmpty())
 			{
-				UScriptStruct* ScriptStructPtr = FindObject<UScriptStruct>(nullptr, *StructPath, false);
+				UScriptStruct* ScriptStructPtr = FindObject<UScriptStruct>(nullptr, *StructPath, EFindObjectFlags::None);
 				if (ScriptStructPtr == nullptr || !ScriptStructPtr->IsChildOf(TBaseStructure<T>::Get()))
 				{
 					Ar.SetError();
diff --git a/Engine/Source/Runtime/Engine/Private/Animation/AnimCompressionTypes.cpp b/Engine/Source/Runtime/Engine/Private/Animation/AnimCompressionTypes.cpp
--- a/Engine/Source/Runtime/Engine/Private/Animation/AnimCompressionTypes.cpp
+++ b/Engine/Source/Runtime/Engine/Private/Animation/AnimCompressionTypes.cpp
@@ -576,7 +576,7 @@ void FCompressibleAnimData::WriteCompressionDataToJSON(TArrayView<FName> Origina
 	{
 		FString JSONString;
 		TSharedRef<TJsonStringWriter<>> Writer = TJsonStringWriter<>::Create(&JSONString);
-		const UEnum* InterpolationEnum = FindObject<UEnum>(nullptr, TEXT("/Script/Engine.EAnimInterpolationType"), true);
+		const UEnum* InterpolationEnum = FindObject<UEnum>(nullptr, TEXT("/Script/Engine.EAnimInterpolationType"), EFindObjectFlags::ExactClass);
 
 		Writer->WriteObjectStart();
 		{
@@ -726,9 +726,9 @@ void FCompressibleAnimData::WriteCompressionDataToJSON(TArrayView<FName> Origina
 							
 						if(FloatCurve.FloatCurve.GetConstRefOfKeys().Num())
 						{
-							const UEnum* CurveInterpolationEnum = FindObject<UEnum>(nullptr, TEXT("/Script/Engine.ERichCurveInterpMode"), true);
-							const UEnum* TangentModeEnum = FindObject<UEnum>(nullptr, TEXT("/Script/Engine.ERichCurveTangentMode"), true);
-							const UEnum* TangentWeightModeEnum = FindObject<UEnum>(nullptr, TEXT("/Script/Engine.ERichCurveTangentWeightMode"), true);
+							const UEnum* CurveInterpolationEnum = FindObject<UEnum>(nullptr, TEXT("/Script/Engine.ERichCurveInterpMode"), EFindObjectFlags::ExactClass);
+							const UEnum* TangentModeEnum = FindObject<UEnum>(nullptr, TEXT("/Script/Engine.ERichCurveTangentMode"), EFindObjectFlags::ExactClass);
+							const UEnum* TangentWeightModeEnum = FindObject<UEnum>(nullptr, TEXT("/Script/Engine.ERichCurveTangentWeightMode"), EFindObjectFlags::ExactClass);
 								
 							Writer->WriteArrayStart(TEXT("key_data"));
 							for (const FRichCurveKey& Key : FloatCurve.FloatCurve.GetConstRefOfKeys())
diff --git a/Engine/Source/Runtime/Engine/Private/DataChannel.cpp b/Engine/Source/Runtime/Engine/Private/DataChannel.cpp
--- a/Engine/Source/Runtime/Engine/Private/DataChannel.cpp
+++ b/Engine/Source/Runtime/Engine/Private/DataChannel.cpp
@@ -5868,7 +5868,7 @@ static void	TestObjectRefSerialize( const TArray<FString>& Args, UWorld* InWorld
 	if (!InWorld || Args.Num() <= 0)
 		return;
 
-	UObject* Object = StaticFindObject( UObject::StaticClass(), NULL, *Args[0], false );
+	UObject* Object = StaticFindObject( UObject::StaticClass(), NULL, *Args[0], EFindObjectFlags::None );
 	if (!Object)
 	{
 		Object = StaticLoadObject( UObject::StaticClass(), NULL, *Args[0], NULL, LOAD_NoWarn );
diff --git a/Engine/Source/Runtime/Engine/Private/Level.cpp b/Engine/Source/Runtime/Engine/Private/Level.cpp
--- a/Engine/Source/Runtime/Engine/Private/Level.cpp
+++ b/Engine/Source/Runtime/Engine/Private/Level.cpp
@@ -4282,7 +4282,7 @@ bool ULevel::SetUseActorFolders(bool bInEnabled, bool bInInteractiveMode)
 
 	if (GetWorldPartition())
 	{
-		UClass* WorldPartitionBuilderCommandletClass = FindObject<UClass>(nullptr, TEXT("/Script/UnrealEd.WorldPartitionBuilderCommandlet"), true);
+		UClass* WorldPartitionBuilderCommandletClass = FindObject<UClass>(nullptr, TEXT("/Script/UnrealEd.WorldPartitionBuilderCommandlet"), EFindObjectFlags::ExactClass);
 		const bool bIsRunningWorldPartitionBuilderCommandlet = WorldPartitionBuilderCommandletClass && GetRunningCommandletClass() && GetRunningCommandletClass()->IsChildOf(WorldPartitionBuilderCommandletClass);
 		if (!bIsRunningWorldPartitionBuilderCommandlet)
 		{
diff --git a/Engine/Source/Runtime/Engine/Private/LevelActor.cpp b/Engine/Source/Runtime/Engine/Private/LevelActor.cpp
--- a/Engine/Source/Runtime/Engine/Private/LevelActor.cpp
+++ b/Engine/Source/Runtime/Engine/Private/LevelActor.cpp
@@ -1588,7 +1588,7 @@ void UWorld::LoadSecondaryLevels(bool bForce, TSet<FName>* PackageNamesToSkip)
 				}
 
 				bool bAlreadyLoaded = false;
-				UPackage* LevelPackage = FindObject<UPackage>(NULL, *StreamingLevelWorldAssetPackageName,true);
+				UPackage* LevelPackage = FindObject<UPackage>(NULL, *StreamingLevelWorldAssetPackageName,EFindObjectFlags::ExactClass);
 				// don't need to do any extra work if the level is already loaded
 				if ( LevelPackage && LevelPackage->IsFullyLoaded() ) 
 				{
diff --git a/Engine/Source/Runtime/Engine/Private/LevelStreaming.cpp b/Engine/Source/Runtime/Engine/Private/LevelStreaming.cpp
--- a/Engine/Source/Runtime/Engine/Private/LevelStreaming.cpp
+++ b/Engine/Source/Runtime/Engine/Private/LevelStreaming.cpp
@@ -1599,7 +1599,7 @@ bool ULevelStreaming::RequestLevel(UWorld* PersistentWorld, bool bAllowLevelLoad
 
 	// Try to find the [to be] loaded package.
 	UWorld* World = nullptr;
-	UPackage* LevelPackage = (UPackage*)StaticFindObjectFast(UPackage::StaticClass(), nullptr, DesiredPackageName, /*bExactClass=*/false, RF_NoFlags, EInternalObjectFlags::Garbage);
+	UPackage* LevelPackage = (UPackage*)StaticFindObjectFast(UPackage::StaticClass(), nullptr, DesiredPackageName, EFindObjectFlags::None, RF_NoFlags, EInternalObjectFlags::Garbage);
 	
 	if (LevelPackage)
 	{
diff --git a/Engine/Source/Runtime/Engine/Private/Materials/MaterialExpressions.cpp b/Engine/Source/Runtime/Engine/Private/Materials/MaterialExpressions.cpp
--- a/Engine/Source/Runtime/Engine/Private/Materials/MaterialExpressions.cpp
+++ b/Engine/Source/Runtime/Engine/Private/Materials/MaterialExpressions.cpp
@@ -13348,7 +13348,7 @@ void UMaterialFunctionInterface::PostLoad()
 		const FString EditorOnlyDataName = MaterialFunctionInterface::GetEditorOnlyDataName(*GetName());
 		if (EditorOnlyData->GetName() != EditorOnlyDataName)
 		{
-			UMaterialFunctionInterfaceEditorOnlyData* CorrectEditorOnlyDataObj = Cast<UMaterialFunctionInterfaceEditorOnlyData>(StaticFindObject(EditorOnlyData->GetClass(), EditorOnlyData->GetOuter(), *EditorOnlyDataName, true));
+			UMaterialFunctionInterfaceEditorOnlyData* CorrectEditorOnlyDataObj = Cast<UMaterialFunctionInterfaceEditorOnlyData>(StaticFindObject(EditorOnlyData->GetClass(), EditorOnlyData->GetOuter(), *EditorOnlyDataName, EFindObjectFlags::ExactClass));
 			if (CorrectEditorOnlyDataObj)
 			{
 				// Copy data to correct EditorOnlyObject
diff --git a/Engine/Source/Runtime/Engine/Private/Materials/MaterialInstanceDynamic.cpp b/Engine/Source/Runtime/Engine/Private/Materials/MaterialInstanceDynamic.cpp
--- a/Engine/Source/Runtime/Engine/Private/Materials/MaterialInstanceDynamic.cpp
+++ b/Engine/Source/Runtime/Engine/Private/Materials/MaterialInstanceDynamic.cpp
@@ -84,7 +84,7 @@ UMaterialInstanceDynamic* UMaterialInstanceDynamic::Create(UMaterialInterface* P
 
 	if (Name != NAME_None)
 	{
-		UMaterialInstanceDynamic* ExistingMID = FindObjectFast<UMaterialInstanceDynamic>(Outer, *Name.ToString(), true);
+		UMaterialInstanceDynamic* ExistingMID = FindObjectFast<UMaterialInstanceDynamic>(Outer, *Name.ToString(), EFindObjectFlags::ExactClass);
 		if (ExistingMID)
 		{
 			bool bRenamed = false;
diff --git a/Engine/Source/Runtime/Engine/Private/Materials/MaterialInterface.cpp b/Engine/Source/Runtime/Engine/Private/Materials/MaterialInterface.cpp
--- a/Engine/Source/Runtime/Engine/Private/Materials/MaterialInterface.cpp
+++ b/Engine/Source/Runtime/Engine/Private/Materials/MaterialInterface.cpp
@@ -108,7 +108,7 @@ static EMaterialGetParameterValueFlags MakeParameterValueFlags(bool bOveriddenOn
 static bool SafeRenameEditorOnlyData(UMaterialInterface* MaterialInterface, const FString& EditorOnlyDataName, bool bLogMessageIfRenamed)
 {
 	UObject* EditorOnlyCurrent = MaterialInterface->GetEditorOnlyData();
-	UObject* EditorOnlyWithIncorrectName = StaticFindObject(/*Class=*/ nullptr, MaterialInterface, *EditorOnlyDataName, true);
+	UObject* EditorOnlyWithIncorrectName = StaticFindObject(/*Class=*/ nullptr, MaterialInterface, *EditorOnlyDataName, EFindObjectFlags::ExactClass);
 	if (EditorOnlyWithIncorrectName && EditorOnlyWithIncorrectName != EditorOnlyCurrent)
 	{
 		if (EditorOnlyWithIncorrectName->IsRooted())
diff --git a/Engine/Source/Runtime/Engine/Private/MorphTools.cpp b/Engine/Source/Runtime/Engine/Private/MorphTools.cpp
--- a/Engine/Source/Runtime/Engine/Private/MorphTools.cpp
+++ b/Engine/Source/Runtime/Engine/Private/MorphTools.cpp
@@ -361,7 +361,7 @@ void FFinishBuildMorphTargetData::ApplyEditorData(USkeletalMesh * SkeletalMesh,
 			}
 
 			//Reuse the morph target if it already exist and was not garbage collect (it can happen if we play with the morph target threshold build options)
-			MorphTarget = FindObjectSafe<UMorphTarget>(SkeletalMesh, *MorphTargetName.ToString(), true);
+			MorphTarget = FindObjectSafe<UMorphTarget>(SkeletalMesh, *MorphTargetName.ToString(), EFindObjectFlags::ExactClass);
 			if (MorphTarget && MorphTarget->HasAnyFlags(RF_BeginDestroyed | RF_FinishDestroyed))
 			{
 				MorphTarget = nullptr;
diff --git a/Engine/Source/Runtime/Engine/Private/PackageMapClient.cpp b/Engine/Source/Runtime/Engine/Private/PackageMapClient.cpp
--- a/Engine/Source/Runtime/Engine/Private/PackageMapClient.cpp
+++ b/Engine/Source/Runtime/Engine/Private/PackageMapClient.cpp
@@ -1227,7 +1227,7 @@ FNetworkGUID UPackageMapClient::InternalLoadObject( FArchive & Ar, UObject *& Ob
 				return NetGUID;
 			}
 
-			Object = StaticFindObject(UObject::StaticClass(), ObjOuter, *ObjectName, false);
+			Object = StaticFindObject(UObject::StaticClass(), ObjOuter, *ObjectName, EFindObjectFlags::None);
 
 			// Try to load package if it wasn't found. Note load package fails if the package is already loaded.
 			if (Object == nullptr && bIsPackage)
diff --git a/Engine/Source/Runtime/Engine/Private/SkeletalMesh.cpp b/Engine/Source/Runtime/Engine/Private/SkeletalMesh.cpp
--- a/Engine/Source/Runtime/Engine/Private/SkeletalMesh.cpp
+++ b/Engine/Source/Runtime/Engine/Private/SkeletalMesh.cpp
@@ -2657,7 +2657,7 @@ bool USkeletalMesh::CommitMeshDescription(
 				{
 					UMorphTarget* MorphTarget = nullptr;
 					// See if object already exists.
-					UObject* Obj = StaticFindObjectFastInternal(nullptr, this, MorphTargetName, true);
+					UObject* Obj = StaticFindObjectFastInternal(nullptr, this, MorphTargetName, EFindObjectFlags::ExactClass);
 					if (Obj)
 					{
 						MorphTarget = Cast<UMorphTarget>(Obj);
diff --git a/Engine/Source/Runtime/Engine/Private/StreamableManager.cpp b/Engine/Source/Runtime/Engine/Private/StreamableManager.cpp
--- a/Engine/Source/Runtime/Engine/Private/StreamableManager.cpp
+++ b/Engine/Source/Runtime/Engine/Private/StreamableManager.cpp
@@ -2045,7 +2045,7 @@ void FStreamableManager::FindInMemory(FSoftObjectPath& InOutTargetName, struct F
 	}
 	else
 	{
-		Asset = StaticFindObject(UObject::StaticClass(), InOutTargetName.GetAssetPath(), /*ExactClass*/false);
+		Asset = StaticFindObject(UObject::StaticClass(), InOutTargetName.GetAssetPath(), EFindObjectFlags::None);
 	}
 
 	if (Asset)
diff --git a/Engine/Source/Runtime/Engine/Private/Tests/Loading/AsyncLoadingTests.cpp b/Engine/Source/Runtime/Engine/Private/Tests/Loading/AsyncLoadingTests.cpp
--- a/Engine/Source/Runtime/Engine/Private/Tests/Loading/AsyncLoadingTests.cpp
+++ b/Engine/Source/Runtime/Engine/Private/Tests/Loading/AsyncLoadingTests.cpp
@@ -76,7 +76,7 @@ bool FAsyncLoadingTestInvalidateExportDuringPreload::RunTest(const FString& Para
 			TestTrue(TEXT("The pacakge should load successfully"), Package != nullptr);
 			
 			// Exclude garbage objects as the GC won't have run yet but invalidated objects shdould be marked as garbage by this point
-			UAsyncLoadingTests_Shared* Object1 = FindObjectFast<UAsyncLoadingTests_Shared>(Package, LoadingTestScope.ObjectName, true/*bExactClass*/, EObjectFlags::RF_MirroredGarbage);
+			UAsyncLoadingTests_Shared* Object1 = FindObjectFast<UAsyncLoadingTests_Shared>(Package, LoadingTestScope.ObjectName, EFindObjectFlags::ExactClass, EObjectFlags::RF_MirroredGarbage);
 			if (bExpectToFindObject)
 			{
 				TestTrue(TEXT("The object should have been loaded"), Object1 != nullptr);
diff --git a/Engine/Source/Runtime/Engine/Private/TickTaskManager.cpp b/Engine/Source/Runtime/Engine/Private/TickTaskManager.cpp
--- a/Engine/Source/Runtime/Engine/Private/TickTaskManager.cpp
+++ b/Engine/Source/Runtime/Engine/Private/TickTaskManager.cpp
@@ -1727,7 +1727,7 @@ public:
 	/** Dumps all tick functions to output device. */
 	void DumpAllTickFunctions(FOutputDevice& Ar, int32& EnabledCount, int32& DisabledCount, bool bEnabled, bool bDisabled)
 	{
-		UEnum* TickGroupEnum = CastChecked<UEnum>(StaticFindObject(UEnum::StaticClass(), nullptr, TEXT("/Script/Engine.ETickingGroup"), true));
+		UEnum* TickGroupEnum = CastChecked<UEnum>(StaticFindObject(UEnum::StaticClass(), nullptr, TEXT("/Script/Engine.ETickingGroup"), EFindObjectFlags::ExactClass));
 		if (bEnabled)
 		{
 			for (TSet<FTickFunction*>::TIterator It(AllEnabledTickFunctions); It; ++It)
diff --git a/Engine/Source/Runtime/Engine/Private/WorldPartition/WorldPartitionLevelStreamingDynamic.cpp b/Engine/Source/Runtime/Engine/Private/WorldPartition/WorldPartitionLevelStreamingDynamic.cpp
--- a/Engine/Source/Runtime/Engine/Private/WorldPartition/WorldPartitionLevelStreamingDynamic.cpp
+++ b/Engine/Source/Runtime/Engine/Private/WorldPartition/WorldPartitionLevelStreamingDynamic.cpp
@@ -247,7 +247,7 @@ bool UWorldPartitionLevelStreamingDynamic::RequestLevel(UWorld* InPersistentWorl
 	FScopeCycleCounterUObject Context(InPersistentWorld);
 
 	// Try to find the package to load
-	UPackage* LevelPackage = (UPackage*)StaticFindObjectFast(UPackage::StaticClass(), nullptr, WorldAssetPackageFName, /*bExactClass =*/ false, RF_NoFlags, EInternalObjectFlags::Garbage);
+	UPackage* LevelPackage = (UPackage*)StaticFindObjectFast(UPackage::StaticClass(), nullptr, WorldAssetPackageFName, EFindObjectFlags::None, RF_NoFlags, EInternalObjectFlags::Garbage);
 	UWorld* FoundWorld = LevelPackage ? UWorld::FindWorldInPackage(LevelPackage) : nullptr;
 	check(!FoundWorld || IsValidChecked(FoundWorld));
 	check(!FoundWorld || FoundWorld->PersistentLevel);
diff --git a/Engine/Source/Runtime/Engine/Private/WorldPartition/WorldPartitionSubsystem.cpp b/Engine/Source/Runtime/Engine/Private/WorldPartition/WorldPartitionSubsystem.cpp
--- a/Engine/Source/Runtime/Engine/Private/WorldPartition/WorldPartitionSubsystem.cpp
+++ b/Engine/Source/Runtime/Engine/Private/WorldPartition/WorldPartitionSubsystem.cpp
@@ -449,7 +449,7 @@ TMap<FActorContainerID, TSet<FGuid>> UWorldPartitionSubsystem::GetFilteredActors
 
 bool UWorldPartitionSubsystem::IsRunningConvertWorldPartitionCommandlet()
 {
-	static UClass* WorldPartitionConvertCommandletClass = FindObject<UClass>(nullptr, TEXT("/Script/UnrealEd.WorldPartitionConvertCommandlet"), true);
+	static UClass* WorldPartitionConvertCommandletClass = FindObject<UClass>(nullptr, TEXT("/Script/UnrealEd.WorldPartitionConvertCommandlet"), EFindObjectFlags::ExactClass);
 	check(WorldPartitionConvertCommandletClass);
 	return GetRunningCommandletClass() && GetRunningCommandletClass()->IsChildOf(WorldPartitionConvertCommandletClass);
 }
diff --git a/Engine/Source/Runtime/Experimental/Iris/Core/Private/Iris/ReplicationState/ReplicationStateDescriptorBuilder.cpp b/Engine/Source/Runtime/Experimental/Iris/Core/Private/Iris/ReplicationState/ReplicationStateDescriptorBuilder.cpp
--- a/Engine/Source/Runtime/Experimental/Iris/Core/Private/Iris/ReplicationState/ReplicationStateDescriptorBuilder.cpp
+++ b/Engine/Source/Runtime/Experimental/Iris/Core/Private/Iris/ReplicationState/ReplicationStateDescriptorBuilder.cpp
@@ -2897,9 +2897,8 @@ SIZE_T FReplicationStateDescriptorBuilder::CreateDescriptorsForClass(FResult& Cr
 
 	// Add special Actor properties
 	{
-		constexpr bool bExactMatch = true;
 		//$IRIS TODO: Probably should use a cached pointer of the Actor class here.
-		UClass* ActorClass = CastChecked<UClass>(StaticFindObject(UClass::StaticClass(), nullptr, TEXT("/Script/Engine.Actor"), bExactMatch));
+		UClass* ActorClass = CastChecked<UClass>(StaticFindObject(UClass::StaticClass(), nullptr, TEXT("/Script/Engine.Actor"), EFindObjectFlags::ExactClass));
 		if (InObjectClass->IsChildOf(ActorClass))
 		{
 			if (const FProperty* NetCullDistanceSquaredProperty = InObjectClass->FindPropertyByName(ReplicationStateDescriptorBuilder_NAME_NetCullDistanceSquared))
diff --git a/Engine/Source/Runtime/Experimental/Iris/Core/Private/Iris/ReplicationSystem/Filtering/ReplicationFiltering.cpp b/Engine/Source/Runtime/Experimental/Iris/Core/Private/Iris/ReplicationSystem/Filtering/ReplicationFiltering.cpp
--- a/Engine/Source/Runtime/Experimental/Iris/Core/Private/Iris/ReplicationSystem/Filtering/ReplicationFiltering.cpp
+++ b/Engine/Source/Runtime/Experimental/Iris/Core/Private/Iris/ReplicationSystem/Filtering/ReplicationFiltering.cpp
@@ -2703,8 +2703,7 @@ void FReplicationFiltering::InitFilters()
 	for (const FNetObjectFilterDefinition& FilterDefinition : FilterDefinitions->GetFilterDefinitions())
 	{
 		constexpr UObject* ClassOuter = nullptr;
-		constexpr bool bMatchExactClass = true;
-		UClass* NetObjectFilterClass = Cast<UClass>(StaticFindObject(UClass::StaticClass(), ClassOuter, ToCStr(FilterDefinition.ClassName.ToString()), bMatchExactClass));
+		UClass* NetObjectFilterClass = Cast<UClass>(StaticFindObject(UClass::StaticClass(), ClassOuter, ToCStr(FilterDefinition.ClassName.ToString()), EFindObjectFlags::ExactClass));
 		if (NetObjectFilterClass == nullptr || !NetObjectFilterClass->IsChildOf(UNetObjectFilter::StaticClass()))
 		{
 			ensureMsgf(NetObjectFilterClass != nullptr, TEXT("NetObjectFilter class is not a NetObjectFilter or could not be found: %s"), ToCStr(FilterDefinition.ClassName.ToString()));
@@ -2714,7 +2713,7 @@ void FReplicationFiltering::InitFilters()
 		UClass* NetObjectFilterConfigClass = nullptr;
 		if (!FilterDefinition.ConfigClassName.IsNone())
 		{
-			NetObjectFilterConfigClass = Cast<UClass>(StaticFindObject(UClass::StaticClass(), ClassOuter, ToCStr(FilterDefinition.ConfigClassName.ToString()), bMatchExactClass));
+			NetObjectFilterConfigClass = Cast<UClass>(StaticFindObject(UClass::StaticClass(), ClassOuter, ToCStr(FilterDefinition.ConfigClassName.ToString()), EFindObjectFlags::ExactClass));
 			if (NetObjectFilterConfigClass == nullptr || !NetObjectFilterConfigClass->IsChildOf(UNetObjectFilterConfig::StaticClass()))
 			{
 				ensureMsgf(NetObjectFilterConfigClass != nullptr, TEXT("NetObjectFilterConfig class is not a NetObjectFilterConfig or could not be found: %s"), ToCStr(FilterDefinition.ConfigClassName.ToString()));
diff --git a/Engine/Source/Runtime/Experimental/Iris/Core/Private/Iris/ReplicationSystem/ObjectReferenceCache.cpp b/Engine/Source/Runtime/Experimental/Iris/Core/Private/Iris/ReplicationSystem/ObjectReferenceCache.cpp
--- a/Engine/Source/Runtime/Experimental/Iris/Core/Private/Iris/ReplicationSystem/ObjectReferenceCache.cpp
+++ b/Engine/Source/Runtime/Experimental/Iris/Core/Private/Iris/ReplicationSystem/ObjectReferenceCache.cpp
@@ -1121,7 +1121,7 @@ ENetObjectReferenceResolveResult FObjectReferenceCache::ResolveObjectReference(c
 			{
 				bool bLocalMustBeMapped = false;
 				UObject* ReplicatedOuter = ResolveObjectReferenceHandleInternal(Reference.GetRefHandle(), ResolveContext, bLocalMustBeMapped);
-				ResolvedObject = ReplicatedOuter ? StaticFindObject(UObject::StaticClass(), ReplicatedOuter, ToCStr(ObjectPath), false) : nullptr;
+				ResolvedObject = ReplicatedOuter ? StaticFindObject(UObject::StaticClass(), ReplicatedOuter, ToCStr(ObjectPath), EFindObjectFlags::None) : nullptr;
 
 				if (ResolvedObject == nullptr)
 				{
@@ -1130,7 +1130,7 @@ ENetObjectReferenceResolveResult FObjectReferenceCache::ResolveObjectReference(c
 			}
 			else
 			{
-				ResolvedObject = StaticFindObject(UObject::StaticClass(), nullptr, *ObjectPath, false);
+				ResolvedObject = StaticFindObject(UObject::StaticClass(), nullptr, *ObjectPath, EFindObjectFlags::None);
 
 				// Try to load package if it wasn't found. Note load package fails if the package is already loaded.
 				if (ResolvedObject == nullptr)
@@ -1138,7 +1138,7 @@ ENetObjectReferenceResolveResult FObjectReferenceCache::ResolveObjectReference(c
 					FPackagePath Path = FPackagePath::FromPackageNameChecked(FPackageName::ObjectPathToPackageName(ObjectPath));
 					if (UPackage* LoadedPackage = LoadPackage(nullptr, Path, LOAD_None))
 					{
-						ResolvedObject = StaticFindObject(UObject::StaticClass(), nullptr, *ObjectPath, false);	
+						ResolvedObject = StaticFindObject(UObject::StaticClass(), nullptr, *ObjectPath, EFindObjectFlags::None);
 					}
 				}
 
diff --git a/Engine/Source/Runtime/Experimental/Iris/Core/Private/Iris/ReplicationSystem/ObjectReplicationBridge.cpp b/Engine/Source/Runtime/Experimental/Iris/Core/Private/Iris/ReplicationSystem/ObjectReplicationBridge.cpp
--- a/Engine/Source/Runtime/Experimental/Iris/Core/Private/Iris/ReplicationSystem/ObjectReplicationBridge.cpp
+++ b/Engine/Source/Runtime/Experimental/Iris/Core/Private/Iris/ReplicationSystem/ObjectReplicationBridge.cpp
@@ -2325,8 +2325,7 @@ void UObjectReplicationBridge::FindClassesInPollPeriodOverrides()
 		if (!PollInfo.Class.IsValid())
 		{
 			constexpr UObject* ClassOuter = nullptr;
-			constexpr bool bExactClass = true;
-			const UClass* Class = CastChecked<UClass>(StaticFindObject(UClass::StaticClass(), ClassOuter, ToCStr(ClassNameAndPollInfo.Key.ToString()), bExactClass));
+			const UClass* Class = CastChecked<UClass>(StaticFindObject(UClass::StaticClass(), ClassOuter, ToCStr(ClassNameAndPollInfo.Key.ToString()), EFindObjectFlags::ExactClass));
 			bFailedToFindClass |= (Class == nullptr);
 			PollInfo.Class = Class;
 		}
diff --git a/Engine/Source/Runtime/Experimental/Iris/Core/Private/Iris/ReplicationSystem/ObjectReplicationBridgeDebugging.cpp b/Engine/Source/Runtime/Experimental/Iris/Core/Private/Iris/ReplicationSystem/ObjectReplicationBridgeDebugging.cpp
--- a/Engine/Source/Runtime/Experimental/Iris/Core/Private/Iris/ReplicationSystem/ObjectReplicationBridgeDebugging.cpp
+++ b/Engine/Source/Runtime/Experimental/Iris/Core/Private/Iris/ReplicationSystem/ObjectReplicationBridgeDebugging.cpp
@@ -497,8 +497,7 @@ void UObjectReplicationBridge::PrintDynamicFilterClassConfig(uint32 ArgTraits)
 				const FString& ClassName = FilterConfig.ClassName.ToString();
 
 				constexpr UObject* ClassOuter = nullptr;
-				constexpr bool bExactClass = true;
-				UClass* Class = Cast<UClass>(StaticFindObject(UClass::StaticClass(), ClassOuter, ToCStr(ClassName), bExactClass));
+				UClass* Class = Cast<UClass>(StaticFindObject(UClass::StaticClass(), ClassOuter, ToCStr(ClassName), EFindObjectFlags::ExactClass));
 				if (Class)
 				{
 					ClassesWithConfig.Add(Class);
diff --git a/Engine/Source/Runtime/LevelSequence/Private/LevelSequenceBindingReference.cpp b/Engine/Source/Runtime/LevelSequence/Private/LevelSequenceBindingReference.cpp
--- a/Engine/Source/Runtime/LevelSequence/Private/LevelSequenceBindingReference.cpp
+++ b/Engine/Source/Runtime/LevelSequence/Private/LevelSequenceBindingReference.cpp
@@ -141,7 +141,7 @@ UObject* ResolveByPath(UObject* InContext, const FString& InObjectPath)
 
 	if (!InObjectPath.IsEmpty())
 	{
-		if (UObject* FoundObject = FindObject<UObject>(InContext, *InObjectPath, false))
+		if (UObject* FoundObject = FindObject<UObject>(InContext, *InObjectPath, EFindObjectFlags::None))
 		{
 			return FoundObject;
 		}
@@ -160,7 +160,7 @@ UObject* ResolveByPath(UObject* InContext, const FString& InObjectPath)
 				for (ULevel* Level : WorldContext->GetLevels())
 				{
 					UPackage* Pkg = Level->GetOutermost();
-					if (UObject* FoundObject = FindObject<UObject>(Pkg, *SubLevelObjPath, false))
+					if (UObject* FoundObject = FindObject<UObject>(Pkg, *SubLevelObjPath, EFindObjectFlags::None))
  					{
  						return FoundObject;
 					}
diff --git a/Engine/Source/Runtime/LiveLinkInterface/Public/LiveLinkTypes.h b/Engine/Source/Runtime/LiveLinkInterface/Public/LiveLinkTypes.h
--- a/Engine/Source/Runtime/LiveLinkInterface/Public/LiveLinkTypes.h
+++ b/Engine/Source/Runtime/LiveLinkInterface/Public/LiveLinkTypes.h
@@ -442,7 +442,7 @@ public:
 			Ar << StructPath;
 			if (!StructPath.IsEmpty())
 			{
-				UScriptStruct* ScriptStructPtr = FindObject<UScriptStruct>(nullptr, *StructPath, false);
+				UScriptStruct* ScriptStructPtr = FindObject<UScriptStruct>(nullptr, *StructPath, EFindObjectFlags::None);
 				if (ScriptStructPtr == nullptr || !ScriptStructPtr->IsChildOf(TBaseStructure<BaseType>::Get()))
 				{
 					Ar.SetError();
diff --git a/Engine/Source/Runtime/NavigationSystem/Private/NavigationSystem.cpp b/Engine/Source/Runtime/NavigationSystem/Private/NavigationSystem.cpp
--- a/Engine/Source/Runtime/NavigationSystem/Private/NavigationSystem.cpp
+++ b/Engine/Source/Runtime/NavigationSystem/Private/NavigationSystem.cpp
@@ -4758,7 +4758,7 @@ ANavigationData* UNavigationSystemV1::CreateNavigationDataInstanceInLevel(const
 			FString StrName = FString::Printf(TEXT("%s-%s"), *(Instance->GetFName().GetPlainNameString()), *(NavConfig.Name.ToString()));
 			// temporary solution to make sure we don't try to change name while there's already
 			// an object with this name
-			UObject* ExistingObject = StaticFindObject(/*Class=*/ nullptr, Instance->GetOuter(), *StrName, true);
+			UObject* ExistingObject = StaticFindObject(/*Class=*/ nullptr, Instance->GetOuter(), *StrName, EFindObjectFlags::ExactClass);
 			while (ExistingObject != nullptr)
 			{
 				ANavigationData* ExistingNavigationData = Cast<ANavigationData>(ExistingObject);
@@ -4771,7 +4771,7 @@ ANavigationData* UNavigationSystemV1::CreateNavigationDataInstanceInLevel(const
 				ExistingObject->Rename(nullptr, nullptr, REN_DontCreateRedirectors | REN_ForceGlobalUnique | REN_DoNotDirty | REN_NonTransactional);
 				// see if there's another one, it does happen when undo/redoing 
 				// nav instance deletion in the editor
-				ExistingObject = StaticFindObject(/*Class=*/ nullptr, Instance->GetOuter(), *StrName, true);
+				ExistingObject = StaticFindObject(/*Class=*/ nullptr, Instance->GetOuter(), *StrName, EFindObjectFlags::ExactClass);
 			}
 
 			// Set descriptive name
diff --git a/Engine/Source/Runtime/TypedElementFramework/Private/Elements/Framework/TypedElementQueryBuilder.cpp b/Engine/Source/Runtime/TypedElementFramework/Private/Elements/Framework/TypedElementQueryBuilder.cpp
--- a/Engine/Source/Runtime/TypedElementFramework/Private/Elements/Framework/TypedElementQueryBuilder.cpp
+++ b/Engine/Source/Runtime/TypedElementFramework/Private/Elements/Framework/TypedElementQueryBuilder.cpp
@@ -21,8 +21,7 @@ namespace UE::Editor::DataStorage::Queries
 
 	const UScriptStruct* TypeOptional(FTopLevelAssetPath Name)
 	{
-		constexpr bool bExactMatch = true;
-		return static_cast<UScriptStruct*>(StaticFindObject(UScriptStruct::StaticClass(), Name, bExactMatch));
+		return static_cast<UScriptStruct*>(StaticFindObject(UScriptStruct::StaticClass(), Name, EFindObjectFlags::ExactClass));
 	}
 
 	const UScriptStruct* operator""_Type(const char* Name, std::size_t NameSize)
