Removing an async loading pipeline flush that is believed to no longer be necessary.

diff --git a/Engine/Plugins/Online/OnlineFramework/Source/Hotfix/Private/OnlineHotfixManager.cpp b/Engine/Plugins/Online/OnlineFramework/Source/Hotfix/Private/OnlineHotfixManager.cpp
--- a/Engine/Plugins/Online/OnlineFramework/Source/Hotfix/Private/OnlineHotfixManager.cpp
+++ b/Engine/Plugins/Online/OnlineFramework/Source/Hotfix/Private/OnlineHotfixManager.cpp
@@ -1420,9 +1420,6 @@ void UOnlineHotfixManager::PatchAssetsFromIniFiles()
 	const FConfigSection* AssetHotfixConfigSection = GConfig->GetSection(TEXT("AssetHotfix"), false, GGameIni);
 	if (AssetHotfixConfigSection != nullptr)
 	{
-		// Flush async loading before modifying GConfig.
-		FlushAsyncLoading();
-
 		// These are the asset types we support patching right now
 		UClass* const PatchableAssetClasses[] = 
 		{ 
