Display in-game invitation as well if the invitation was sent from platform-only friend.

diff --git a/Engine/Plugins/Online/OnlineFramework/Source/Party/Private/SocialToolkit.cpp b/Engine/Plugins/Online/OnlineFramework/Source/Party/Private/SocialToolkit.cpp
--- a/Engine/Plugins/Online/OnlineFramework/Source/Party/Private/SocialToolkit.cpp
+++ b/Engine/Plugins/Online/OnlineFramework/Source/Party/Private/SocialToolkit.cpp
@@ -1068,7 +1068,7 @@ void USocialToolkit::HandleAcceptFriendInviteComplete(int32 LocalUserNum, bool b
 
 const bool USocialToolkit::IsInviteAllowedFromUser(const USocialUser& User, const TSharedRef<const IOnlinePartyJoinInfo>& InvitePtr) const
 {
-  return User.IsFriend(ESocialSubsystem::Primary);
+  return User.IsFriend();
 }
 
 void USocialToolkit::HandlePartyInviteReceived(const FUniqueNetId& LocalUserId, const IOnlinePartyJoinInfo& Invite)
