Fix for social interactions not working correctly in PIE.

This was noticed when trying to debug the Block modal input bug. Without this change I could not get the Block modal to show in PIE.

Slack discussion: https://epicgames.slack.com/archives/C04A9V2R03A/p1749063938176189

#tests Compiled DevelopmentEditor Win64 / Ran a PIE client without fix and noticed in the debugger FSocialInteractionHandle comparisons were failing due to different wrapper addresses. Applied fix, retested, and verified the comparisons succeeded. Francisco Hart noticed this in a previous build.
#rnx

diff --git a/Engine/Plugins/Online/OnlineFramework/Source/Party/Public/Interactions/SocialInteractionMacros.h b/Engine/Plugins/Online/OnlineFramework/Source/Party/Public/Interactions/SocialInteractionMacros.h
--- a/Engine/Plugins/Online/OnlineFramework/Source/Party/Public/Interactions/SocialInteractionMacros.h
+++ b/Engine/Plugins/Online/OnlineFramework/Source/Party/Public/Interactions/SocialInteractionMacros.h
@@ -51,7 +51,7 @@ private:
 class FSocialInteraction_##InteractionName	\
 {	\
 public:	\
-	static FSocialInteractionHandle GetHandle()	\
+	static APIMacro FSocialInteractionHandle GetHandle()	\
 	{	\
 		static const TSocialInteractionWrapper<FSocialInteraction_##InteractionName> InteractionWrapper;	\
 		return InteractionWrapper.GetHandle();	\
