Fix crash when attempting to block users on certain platforms.

'APIMacro' missing on this function doesn't cause any compiler issues, but it leads to unexpected behavior in various parts of the code. This is similar to the fix for 43266600.

#tests Compiled Test Client / Ran a test client and verified I could 100% crash the game blocking users. Applied fix and verified I could not crash the game.
#rnx

diff --git a/Engine/Plugins/Online/OnlineFramework/Source/Party/Public/Interactions/SocialInteractionMacros.h b/Engine/Plugins/Online/OnlineFramework/Source/Party/Public/Interactions/SocialInteractionMacros.h
--- a/Engine/Plugins/Online/OnlineFramework/Source/Party/Public/Interactions/SocialInteractionMacros.h
+++ b/Engine/Plugins/Online/OnlineFramework/Source/Party/Public/Interactions/SocialInteractionMacros.h
@@ -73,7 +73,7 @@ public:	\
 		return false;	\
 	}	\
 	\
-	static void ExecuteInteractionWithContext(USocialUser& User, const TMap<FString, FString>& AnalyticsContext)	\
+	static APIMacro void ExecuteInteractionWithContext(USocialUser& User, const TMap<FString, FString>& AnalyticsContext)	\
 	{	\
 		User.WithContext(AnalyticsContext, [](USocialUser& InUser) {	\
 			ExecuteInteraction(InUser);	\
