- Fix SocialToolkit not binding to OSS OnBlocklistChanged

[RN] Social Toolkit now correctly reacts to block list changes made in platform UI, broadcast by the OSS OnBlocklistChanged event

diff --git a/Engine/Plugins/Online/OnlineFramework/Source/Party/Private/SocialToolkit.cpp b/Engine/Plugins/Online/OnlineFramework/Source/Party/Private/SocialToolkit.cpp
--- a/Engine/Plugins/Online/OnlineFramework/Source/Party/Private/SocialToolkit.cpp
+++ b/Engine/Plugins/Online/OnlineFramework/Source/Party/Private/SocialToolkit.cpp
@@ -615,6 +615,7 @@ void USocialToolkit::OnOwnerLoggedIn()
 
 				FriendsInterface->AddOnBlockedPlayerCompleteDelegate_Handle(LocalUserNum, FOnBlockedPlayerCompleteDelegate::CreateUObject(this, &USocialToolkit::HandleBlockPlayerComplete, SubsystemType));
 				FriendsInterface->AddOnUnblockedPlayerCompleteDelegate_Handle(LocalUserNum, FOnUnblockedPlayerCompleteDelegate::CreateUObject(this, &USocialToolkit::HandleUnblockPlayerComplete, SubsystemType));
+				FriendsInterface->AddOnBlockListChangeDelegate_Handle(LocalUserNum, FOnBlockListChangeDelegate::CreateUObject(this, &USocialToolkit::HandleBlockListChanged, SubsystemType));
 
 				FriendsInterface->AddOnRecentPlayersAddedDelegate_Handle(FOnRecentPlayersAddedDelegate::CreateUObject(this, &USocialToolkit::HandleRecentPlayersAdded, SubsystemType));
 			}
@@ -672,6 +673,7 @@ void USocialToolkit::OnOwnerLoggedOut()
 
 				FriendsInterface->ClearOnBlockedPlayerCompleteDelegates(LocalUserNum, this);
 				FriendsInterface->ClearOnUnblockedPlayerCompleteDelegates(LocalUserNum, this);
+				FriendsInterface->ClearOnBlockListChangeDelegates(LocalUserNum, this);
 
 				FriendsInterface->ClearOnRecentPlayersAddedDelegates(this);
 
@@ -837,7 +839,9 @@ void USocialToolkit::HandleQueryBlockedPlayersComplete(const FUniqueNetId& UserI
 
 			TArray<TSharedRef<FOnlineBlockedPlayer>> BlockedPlayers;
 			FriendsInterface->GetBlockedPlayers(UserId, BlockedPlayers);
-			ProcessUserList(BlockedPlayers, SubsystemType, OnUserBlocked());
+
+			ProcessOSSBlockedPlayers(SubsystemType, BlockedPlayers);
+
 			OnQueryBlockedPlayersSuccess(SubsystemType, BlockedPlayers);
 		}
 		else
@@ -1170,6 +1174,50 @@ void USocialToolkit::HandleUnblockPlayerComplete(int32 LocalUserNum, bool bWasSu
 	OnUnblockPlayerComplete(LocalUserNum, bWasSuccessful, UnblockedPlayerId, ListName, ErrorStr, SubsystemType);
 }
 
+void USocialToolkit::HandleBlockListChanged(int32 LocalUserNum, const FString& ListName, ESocialSubsystem SubsystemType)
+{
+	if (LocalUserNum == GetLocalUserNum())
+	{
+		UE_LOG(LogParty, Log, TEXT("SocialToolkit [%d] Subsystem=[%s] Blocklist changed"), GetLocalUserNum(), ToString(SubsystemType));
+
+		TArray<TSharedRef<FOnlineBlockedPlayer>> BlockedPlayers;
+		IOnlineFriendsPtr FriendsInterface = Online::GetFriendsInterfaceChecked(GetWorld(), USocialManager::GetSocialOssName(SubsystemType));
+		FriendsInterface->GetBlockedPlayers(*GetLocalUserNetId(SubsystemType), BlockedPlayers);
+
+		ProcessOSSBlockedPlayers(SubsystemType, BlockedPlayers);
+	}
+}
+
+void USocialToolkit::ProcessOSSBlockedPlayers(ESocialSubsystem SubsystemType, const TArray<TSharedRef<FOnlineBlockedPlayer>>& InBlockedPlayers)
+{
+	TSet<FUniqueNetIdRepl> InBlockedPlayerNetIds;
+	Algo::Transform(InBlockedPlayers, InBlockedPlayerNetIds, [](const TSharedRef<FOnlineBlockedPlayer>& BlockedPlayer) { return BlockedPlayer->GetUserId(); });
+
+	TSet<FUniqueNetIdRepl>& CurrentBlockedUserNetIds = BlockedUsers.FindOrAdd(SubsystemType);
+	TSet<FUniqueNetIdRepl> UnblockedNetIds = CurrentBlockedUserNetIds.Difference(InBlockedPlayerNetIds);
+
+	UE_LOG(LogParty, Log, TEXT("SocialToolkit [%d] Subsystem=[%s] BlockedPlayers=[%s] UnblockedPlayers=[%s]"),
+		GetLocalUserNum(),
+		LexToString(SubsystemType),
+		*FString::JoinBy(InBlockedPlayerNetIds, TEXT(","), &FUniqueNetIdRepl::ToDebugString),
+		*FString::JoinBy(UnblockedNetIds, TEXT(","), &FUniqueNetIdRepl::ToDebugString));
+
+	CurrentBlockedUserNetIds = MoveTemp(InBlockedPlayerNetIds);
+
+	for (const FUniqueNetIdRepl& UnblockedNetId : UnblockedNetIds)
+	{
+		if (TWeakObjectPtr<USocialUser>* WeakUnblockedUser = UsersBySubsystemIds.Find(UnblockedNetId))
+		{
+			if (USocialUser* UnblockedUser = WeakUnblockedUser->Get())
+			{
+				UnblockedUser->NotifyUserUnblocked(SubsystemType);
+			}
+		}
+	}
+
+	ProcessUserList(InBlockedPlayers, SubsystemType, OnUserBlocked());
+}
+
 //@todo DanH recent players: Where is the line for this between backend and game to update this stuff? #required
 //		Seems like I should just be able to get an event for OnRecentPlayersAdded or even a full OnRecentPlayersListRefreshed from IOnlineFriends
 void USocialToolkit::HandlePartyMemberExited(const FUniqueNetId& LocalUserId, const FOnlinePartyId& PartyId, const FUniqueNetId& MemberId, const EMemberExitedReason Reason)
diff --git a/Engine/Plugins/Online/OnlineFramework/Source/Party/Public/SocialToolkit.h b/Engine/Plugins/Online/OnlineFramework/Source/Party/Public/SocialToolkit.h
--- a/Engine/Plugins/Online/OnlineFramework/Source/Party/Public/SocialToolkit.h
+++ b/Engine/Plugins/Online/OnlineFramework/Source/Party/Public/SocialToolkit.h
@@ -283,6 +283,8 @@ private:	// Handlers
 
 	UE_API void HandleBlockPlayerComplete(int32 LocalUserNum, bool bWasSuccessful, const FUniqueNetId& BlockedPlayerID, const FString& ListName, const FString& ErrorStr, ESocialSubsystem SubsystemType);
 	UE_API void HandleUnblockPlayerComplete(int32 LocalUserNum, bool bWasSuccessful, const FUniqueNetId& UnblockedPlayerID, const FString& ListName, const FString& ErrorStr, ESocialSubsystem SubsystemType);
+
+	UE_API void HandleBlockListChanged(int32 LocalUserNum, const FString& ListName, ESocialSubsystem SubsystemType);
 	
 	UE_API void HandleQueryPrimaryUserIdMappingComplete(bool bWasSuccessful, const FUniqueNetId& RequestingUserId, const FString& DisplayName, const FUniqueNetId& IdentifiedUserId, const FString& Error);
 
@@ -316,6 +318,10 @@ private:
 
 	UPROPERTY()
 	TArray<TObjectPtr<USocialUser>> AllUsers;
+
+	void ProcessOSSBlockedPlayers(ESocialSubsystem SubsystemType, const TArray<TSharedRef<FOnlineBlockedPlayer>>& BlockedPlayers);
+	TMap<ESocialSubsystem, TSet<FUniqueNetIdRepl>> BlockedUsers;
+
 	TMap<FUniqueNetIdRepl, TWeakObjectPtr<USocialUser>> UsersBySubsystemIds;
 
 	UPROPERTY(config)
