Add an accessor for the result of the most recent patch check.

diff --git a/Engine/Plugins/Online/OnlineFramework/Source/PatchCheck/Private/PatchCheck.cpp b/Engine/Plugins/Online/OnlineFramework/Source/PatchCheck/Private/PatchCheck.cpp
--- a/Engine/Plugins/Online/OnlineFramework/Source/PatchCheck/Private/PatchCheck.cpp
+++ b/Engine/Plugins/Online/OnlineFramework/Source/PatchCheck/Private/PatchCheck.cpp
@@ -437,6 +437,7 @@ void FPatchCheck::PatchCheckComplete(EPatchCheckResult PatchResult)
 		StrongStats->OnPatchCheckComplete(PatchResult);
 	}
 
+	LastPatchCheckResult = PatchResult;
 	OnComplete.Broadcast(PatchResult);
 	bIsCheckInProgress = false;
 }
diff --git a/Engine/Plugins/Online/OnlineFramework/Source/PatchCheck/Public/PatchCheck.h b/Engine/Plugins/Online/OnlineFramework/Source/PatchCheck/Public/PatchCheck.h
--- a/Engine/Plugins/Online/OnlineFramework/Source/PatchCheck/Public/PatchCheck.h
+++ b/Engine/Plugins/Online/OnlineFramework/Source/PatchCheck/Public/PatchCheck.h
@@ -72,6 +72,7 @@ public:
 	UE_API void RemoveEnvironmentWantsPatchCheckBackCompatDelegate(FName Tag);
 
 	FOnPatchCheckComplete& GetOnComplete() { return OnComplete; }
+	TOptional<EPatchCheckResult> GetLastPatchCheckResult() const { return LastPatchCheckResult; }
 
 	UE_API void RegisterStatsCollector(const TSharedPtr<IPatchCheckStatsCollector>& InStats);
 
@@ -113,6 +114,7 @@ protected:
 	TWeakPtr<IPatchCheckStatsCollector> Stats;
 	TOptional<FPreviousSuccessInfo> PreviousPlatformSuccess;
 	TOptional<FPreviousSuccessInfo> PreviousOnlineServiceSuccess;
+	TOptional<EPatchCheckResult> LastPatchCheckResult;
 
 	FOnPatchCheckComplete OnComplete;
 
