[OnlineFramework]
* Changed DECLARE_SOCIAL_INTERACTION to only declare type/methods and added DEFINE_SOCIAL_INTERACTION which will define the methods. This to fix the issue with singletons in headers causing issues over dll boundaries (merged modules in our case)

diff --git a/Engine/Plugins/Online/OnlineFramework/Source/Party/Private/Interactions/CoreInteractions.cpp b/Engine/Plugins/Online/OnlineFramework/Source/Party/Private/Interactions/CoreInteractions.cpp
--- a/Engine/Plugins/Online/OnlineFramework/Source/Party/Private/Interactions/CoreInteractions.cpp
+++ b/Engine/Plugins/Online/OnlineFramework/Source/Party/Private/Interactions/CoreInteractions.cpp
@@ -12,6 +12,8 @@
 // AddFriend
 //////////////////////////////////////////////////////////////////////////
 
+DEFINE_SOCIAL_INTERACTION(AddFriend)
+
 FText FSocialInteraction_AddFriend::GetDisplayName(const USocialUser& User)
 {
 	if (User.IsFriend(ESocialSubsystem::Platform))
@@ -40,6 +42,8 @@ void FSocialInteraction_AddFriend::ExecuteInteraction(USocialUser& User)
 // AddPlatformFriend
 //////////////////////////////////////////////////////////////////////////
 
+DEFINE_SOCIAL_INTERACTION(AddPlatformFriend)
+
 FText FSocialInteraction_AddPlatformFriend::GetDisplayName(const USocialUser& User)
 {
 	const FName PlatformOssName = USocialManager::GetSocialOssName(ESocialSubsystem::Platform);
@@ -73,6 +77,7 @@ void FSocialInteraction_AddPlatformFriend::ExecuteInteraction(USocialUser& User)
 // RemoveFriend
 //////////////////////////////////////////////////////////////////////////
 
+DEFINE_SOCIAL_INTERACTION(RemoveFriend)
 
 FText FSocialInteraction_RemoveFriend::GetDisplayName(const USocialUser& User)
 {
@@ -99,6 +104,8 @@ void FSocialInteraction_RemoveFriend::ExecuteInteraction(USocialUser& User)
 // AcceptFriendInvite
 //////////////////////////////////////////////////////////////////////////
 
+DEFINE_SOCIAL_INTERACTION(AcceptFriendInvite)
+
 FText FSocialInteraction_AcceptFriendInvite::GetDisplayName(const USocialUser& User)
 {
 	return LOCTEXT("AcceptFriendInvite", "Accept");
@@ -123,6 +130,8 @@ void FSocialInteraction_AcceptFriendInvite::ExecuteInteraction(USocialUser& User
 // RejectFriendInvite
 //////////////////////////////////////////////////////////////////////////
 
+DEFINE_SOCIAL_INTERACTION(RejectFriendInvite)
+
 FText FSocialInteraction_RejectFriendInvite::GetDisplayName(const USocialUser& User)
 {
 	return LOCTEXT("RejectFriendInvite", "Reject");
@@ -147,6 +156,8 @@ void FSocialInteraction_RejectFriendInvite::ExecuteInteraction(USocialUser& User
 // Block
 //////////////////////////////////////////////////////////////////////////
 
+DEFINE_SOCIAL_INTERACTION(Block)
+
 FText FSocialInteraction_Block::GetDisplayName(const USocialUser& User)
 {
 	return LOCTEXT("BlockUser", "Block");
@@ -171,6 +182,8 @@ void FSocialInteraction_Block::ExecuteInteraction(USocialUser& User)
 // Unblock
 //////////////////////////////////////////////////////////////////////////
 
+DEFINE_SOCIAL_INTERACTION(Unblock)
+
 FText FSocialInteraction_Unblock::GetDisplayName(const USocialUser& User)
 {
 	return LOCTEXT("UnblockUser", "Unblock");
@@ -195,6 +208,7 @@ void FSocialInteraction_Unblock::ExecuteInteraction(USocialUser& User)
 // PrivateMessage
 //////////////////////////////////////////////////////////////////////////
 
+DEFINE_SOCIAL_INTERACTION(PrivateMessage)
 
 FText FSocialInteraction_PrivateMessage::GetDisplayName(const USocialUser& User)
 {
@@ -225,6 +239,8 @@ void FSocialInteraction_PrivateMessage::ExecuteInteraction(USocialUser& User)
 // ShowPlatformProfile
 //////////////////////////////////////////////////////////////////////////
 
+DEFINE_SOCIAL_INTERACTION(ShowPlatformProfile)
+
 FText FSocialInteraction_ShowPlatformProfile::GetDisplayName(const USocialUser& User)
 {
 	return LOCTEXT("ShowPlatformProfile", "View Profile");
diff --git a/Engine/Plugins/Online/OnlineFramework/Source/Party/Private/Interactions/PartyInteractions.cpp b/Engine/Plugins/Online/OnlineFramework/Source/Party/Private/Interactions/PartyInteractions.cpp
--- a/Engine/Plugins/Online/OnlineFramework/Source/Party/Private/Interactions/PartyInteractions.cpp
+++ b/Engine/Plugins/Online/OnlineFramework/Source/Party/Private/Interactions/PartyInteractions.cpp
@@ -13,6 +13,8 @@
 // InviteToParty
 //////////////////////////////////////////////////////////////////////////
 
+DEFINE_SOCIAL_INTERACTION(InviteToParty)
+
 FText FSocialInteraction_InviteToParty::GetDisplayName(const USocialUser& User)
 {
 	return LOCTEXT("InviteToParty", "Invite to Party");
@@ -37,6 +39,8 @@ void FSocialInteraction_InviteToParty::ExecuteInteraction(USocialUser& User)
 // JoinParty
 //////////////////////////////////////////////////////////////////////////
 
+DEFINE_SOCIAL_INTERACTION(JoinParty)
+
 FText FSocialInteraction_JoinParty::GetDisplayName(const USocialUser& User)
 {
 	return LOCTEXT("JoinParty", "Join Party");
@@ -62,6 +66,8 @@ void FSocialInteraction_JoinParty::ExecuteInteraction(USocialUser& User)
 // RequestToJoinParty
 //////////////////////////////////////////////////////////////////////////
 
+DEFINE_SOCIAL_INTERACTION(RequestToJoinParty)
+
 FText FSocialInteraction_RequestToJoinParty::GetDisplayName(const USocialUser& User)
 {
 	return LOCTEXT("RequestToJoinParty", "Request to Join");
@@ -86,6 +92,8 @@ void FSocialInteraction_RequestToJoinParty::ExecuteInteraction(USocialUser& User
 // AcceptJoinRequest
 //////////////////////////////////////////////////////////////////////////
 
+DEFINE_SOCIAL_INTERACTION(AcceptJoinRequest)
+
 FText FSocialInteraction_AcceptJoinRequest::GetDisplayName(const USocialUser& User)
 {
 	return LOCTEXT("AcceptJoinRequest", "Accept Request");
@@ -110,6 +118,8 @@ void FSocialInteraction_AcceptJoinRequest::ExecuteInteraction(USocialUser& User)
 // DismissJoinRequest
 //////////////////////////////////////////////////////////////////////////
 
+DEFINE_SOCIAL_INTERACTION(DismissJoinRequest)
+
 FText FSocialInteraction_DismissJoinRequest::GetDisplayName(const USocialUser& User)
 {
 	return LOCTEXT("DismissJoinRequest", "Dismiss Request");
@@ -134,6 +144,8 @@ void FSocialInteraction_DismissJoinRequest::ExecuteInteraction(USocialUser& User
 // AcceptPartyInvite
 //////////////////////////////////////////////////////////////////////////
 
+DEFINE_SOCIAL_INTERACTION(AcceptPartyInvite)
+
 FText FSocialInteraction_AcceptPartyInvite::GetDisplayName(const USocialUser& User)
 {
 	return LOCTEXT("AcceptPartyInvite", "Accept Invite");
@@ -159,6 +171,8 @@ void FSocialInteraction_AcceptPartyInvite::ExecuteInteraction(USocialUser& User)
 // RejectPartyInvite
 //////////////////////////////////////////////////////////////////////////
 
+DEFINE_SOCIAL_INTERACTION(RejectPartyInvite)
+
 FText FSocialInteraction_RejectPartyInvite::GetDisplayName(const USocialUser& User)
 {
 	return LOCTEXT("RejectPartyInvite", "Reject Invite");
@@ -183,6 +197,8 @@ void FSocialInteraction_RejectPartyInvite::ExecuteInteraction(USocialUser& User)
 // LeaveParty
 //////////////////////////////////////////////////////////////////////////
 
+DEFINE_SOCIAL_INTERACTION(LeaveParty)
+
 FText FSocialInteraction_LeaveParty::GetDisplayName(const USocialUser& User)
 {
 	return LOCTEXT("LeaveParty", "Leave Party");
@@ -222,6 +238,8 @@ void FSocialInteraction_LeaveParty::ExecuteInteraction(USocialUser& User)
 // KickPartyMember
 //////////////////////////////////////////////////////////////////////////
 
+DEFINE_SOCIAL_INTERACTION(KickPartyMember)
+
 FText FSocialInteraction_KickPartyMember::GetDisplayName(const USocialUser& User)
 {
 	return LOCTEXT("KickPartyMember", "Kick");
@@ -253,6 +271,8 @@ void FSocialInteraction_KickPartyMember::ExecuteInteraction(USocialUser& User)
 // PromoteToPartyLeader
 //////////////////////////////////////////////////////////////////////////
 
+DEFINE_SOCIAL_INTERACTION(PromoteToPartyLeader)
+
 FText FSocialInteraction_PromoteToPartyLeader::GetDisplayName(const USocialUser& User)
 {
 	return LOCTEXT("PromoteToPartyLeader", "Promote");
diff --git a/Engine/Plugins/Online/OnlineFramework/Source/Party/Public/Interactions/SocialInteractionMacros.h b/Engine/Plugins/Online/OnlineFramework/Source/Party/Public/Interactions/SocialInteractionMacros.h
--- a/Engine/Plugins/Online/OnlineFramework/Source/Party/Public/Interactions/SocialInteractionMacros.h
+++ b/Engine/Plugins/Online/OnlineFramework/Source/Party/Public/Interactions/SocialInteractionMacros.h
@@ -47,24 +47,36 @@ private:
 // Helper macros for declaring a social interaction class
 // Establishes boilerplate behavior and declares all functions the user is required to provide
 
+
 #define DECLARE_SOCIAL_INTERACTION_EXPORT(APIMacro, InteractionName)	\
-class FSocialInteraction_##InteractionName	\
-{	\
-public:	\
-	static APIMacro FSocialInteractionHandle GetHandle()	\
+	class FSocialInteraction_##InteractionName	\
+	{	\
+	public:	\
+		static APIMacro FSocialInteractionHandle GetHandle();	\
+		static FName GetInteractionName();  \
+		static APIMacro FText GetDisplayName(const USocialUser& User);	\
+		static APIMacro FString GetSlashCommandToken();	\
+		static bool IsAvailable(const USocialUser& User);	\
+		static APIMacro void ExecuteInteractionWithContext(USocialUser& User, const TMap<FString, FString>& AnalyticsContext);	\
+		static APIMacro void ExecuteInteraction(USocialUser& User);	\
+		static FOnCustomIsInteractionAvailable& OnCustomIsInteractionAvailable();	\
+	private: \
+		static APIMacro bool CanExecute(const USocialUser& User); \
+	}
+
+
+
+#define DEFINE_SOCIAL_INTERACTION(InteractionName)	\
+	FSocialInteractionHandle FSocialInteraction_##InteractionName::GetHandle()	\
 	{	\
 		static const TSocialInteractionWrapper<FSocialInteraction_##InteractionName> InteractionWrapper;	\
 		return InteractionWrapper.GetHandle();	\
 	}	\
-	static FName GetInteractionName()  \
+	FName FSocialInteraction_##InteractionName::GetInteractionName()  \
 	{ \
 		return #InteractionName; \
 	} \
-	\
-	static APIMacro FText GetDisplayName(const USocialUser& User);	\
-	static APIMacro FString GetSlashCommandToken();	\
-	\
-	static bool IsAvailable(const USocialUser& User)	\
+	bool FSocialInteraction_##InteractionName::IsAvailable(const USocialUser& User)	\
 	{	\
 		if (CanExecute(User))	\
 		{	\
@@ -72,23 +84,17 @@ public:	\
 		}	\
 		return false;	\
 	}	\
-	\
-	static APIMacro void ExecuteInteractionWithContext(USocialUser& User, const TMap<FString, FString>& AnalyticsContext)	\
+	void FSocialInteraction_##InteractionName::ExecuteInteractionWithContext(USocialUser& User, const TMap<FString, FString>& AnalyticsContext)	\
 	{	\
 		User.WithContext(AnalyticsContext, [](USocialUser& InUser) {	\
 			ExecuteInteraction(InUser);	\
 		});	\
 	}	\
-	\
-	static APIMacro void ExecuteInteraction(USocialUser& User);	\
-	\
-	static FOnCustomIsInteractionAvailable& OnCustomIsInteractionAvailable()	\
+	FOnCustomIsInteractionAvailable& FSocialInteraction_##InteractionName::OnCustomIsInteractionAvailable()	\
 	{	\
 		static FOnCustomIsInteractionAvailable CustomAvailabilityCheckDelegate;	\
 		return CustomAvailabilityCheckDelegate;	\
 	}	\
-private: \
-	static APIMacro bool CanExecute(const USocialUser& User); \
-}
+
 
 #define DECLARE_SOCIAL_INTERACTION(InteractionName) DECLARE_SOCIAL_INTERACTION_EXPORT(, InteractionName)
