Add friends changed delegate to socialtoolkit, to repond to changes to platfrom friends.

diff --git a/Engine/Plugins/Online/OnlineFramework/Source/Party/Private/SocialToolkit.cpp b/Engine/Plugins/Online/OnlineFramework/Source/Party/Private/SocialToolkit.cpp
--- a/Engine/Plugins/Online/OnlineFramework/Source/Party/Private/SocialToolkit.cpp
+++ b/Engine/Plugins/Online/OnlineFramework/Source/Party/Private/SocialToolkit.cpp
@@ -607,7 +607,8 @@ void USocialToolkit::OnOwnerLoggedIn()
 			{
 				FriendsInterface->AddOnFriendRemovedDelegate_Handle(FOnFriendRemovedDelegate::CreateUObject(this, &USocialToolkit::HandleFriendRemoved, SubsystemType));
 				FriendsInterface->AddOnDeleteFriendCompleteDelegate_Handle(LocalUserNum, FOnDeleteFriendCompleteDelegate::CreateUObject(this, &USocialToolkit::HandleDeleteFriendComplete, SubsystemType));
-
+				FriendsInterface->AddOnFriendsChangeDelegate_Handle(LocalUserNum, FOnFriendsChangeDelegate::CreateUObject(this, &USocialToolkit::HandleFriendsChanged, SubsystemType));
+ 
 				FriendsInterface->AddOnInviteReceivedDelegate_Handle(FOnInviteReceivedDelegate::CreateUObject(this, &USocialToolkit::HandleFriendInviteReceived, SubsystemType));
 				FriendsInterface->AddOnInviteAcceptedDelegate_Handle(FOnInviteAcceptedDelegate::CreateUObject(this, &USocialToolkit::HandleFriendInviteAccepted, SubsystemType));
 				FriendsInterface->AddOnInviteRejectedDelegate_Handle(FOnInviteRejectedDelegate::CreateUObject(this, &USocialToolkit::HandleFriendInviteRejected, SubsystemType));
@@ -1050,7 +1051,15 @@ void USocialToolkit::HandleFriendRemoved(const FUniqueNetId& LocalUserId, const
 		}
 	}
 }
+void USocialToolkit::HandleFriendsChanged(ESocialSubsystem SubsystemType)
+{
+	UE_LOG(LogParty, Log, TEXT("SocialToolkit [%d] quietly refreshing Friend list due to an external requery of the friends list."), GetLocalUserNum());
 
+	TArray<TSharedRef<FOnlineFriend>> FriendsList;
+	IOnlineFriendsPtr FriendsInterface = Online::GetFriendsInterfaceChecked(GetWorld(), USocialManager::GetSocialOssName(SubsystemType));
+	FriendsInterface->GetFriendsList(GetLocalUserNum(), FriendListToQuery, FriendsList);
+	ProcessUserList(FriendsList, SubsystemType, OnFriendshipEstablished());
+}
 void USocialToolkit::HandleDeleteFriendComplete(int32 InLocalUserNum, bool bWasSuccessful, const FUniqueNetId& DeletedFriendId, const FString& ListName, const FString& ErrorStr, ESocialSubsystem SubsystemType)
 {
 	if (bWasSuccessful && InLocalUserNum == GetLocalUserNum())
diff --git a/Engine/Plugins/Online/OnlineFramework/Source/Party/Public/SocialToolkit.h b/Engine/Plugins/Online/OnlineFramework/Source/Party/Public/SocialToolkit.h
--- a/Engine/Plugins/Online/OnlineFramework/Source/Party/Public/SocialToolkit.h
+++ b/Engine/Plugins/Online/OnlineFramework/Source/Party/Public/SocialToolkit.h
@@ -274,7 +274,8 @@ private:	// Handlers
 	UE_API void HandleFriendRemoved(const FUniqueNetId& LocalUserId, const FUniqueNetId& FormerFriendId, ESocialSubsystem SubsystemType);
 
 	UE_API void HandleDeleteFriendComplete(int32 LocalPlayer, bool bWasSuccessful, const FUniqueNetId& FormerFriendId, const FString& ListName, const FString& ErrorStr, ESocialSubsystem SubsystemType);
-	
+	UE_API void HandleFriendsChanged(ESocialSubsystem SubsystemType);
+
 	UE_API void HandlePartyInviteReceived(const FUniqueNetId& LocalUserId, const IOnlinePartyJoinInfo& Invite);
 	UE_API void HandlePartyInviteRemoved(const FUniqueNetId& LocalUserId, const IOnlinePartyJoinInfo& Invite, EPartyInvitationRemovedReason Reason);
 #if PARTY_PLATFORM_INVITE_PERMISSIONS
