Added Online Services Presence usage CVar, and code to enable it via commandline parameter

diff --git a/Engine/Plugins/Online/OnlineFramework/Source/Party/Private/PartyModule.cpp b/Engine/Plugins/Online/OnlineFramework/Source/Party/Private/PartyModule.cpp
--- a/Engine/Plugins/Online/OnlineFramework/Source/Party/Private/PartyModule.cpp
+++ b/Engine/Plugins/Online/OnlineFramework/Source/Party/Private/PartyModule.cpp
@@ -2,6 +2,8 @@
 
 #include "PartyModule.h"
 
+#include "HAL/IConsoleManager.h"
+
 IMPLEMENT_MODULE(FPartyModule, Party);
 
 DEFINE_LOG_CATEGORY(LogParty);
@@ -18,3 +20,12 @@ void FPartyModule::ShutdownModule()
 {
 }
 
+#if !UE_BUILD_SHIPPING
+	/**
+		* CVar to enable usage of Online Services Presence
+		*/
+	static TAutoConsoleVariable<int32> CVar_OnlineFramework_Party_IsOnlineServicesPresenceEnabled(
+		TEXT("onlineframework.party.isonlineservicespresenceenabled"),
+		0,
+		TEXT("Whether relevant systems should use OnlineServices for Presence operations"));
+#endif
\ No newline at end of file
