Added options to various Party systems to use OnlineServices for Presence and DisplayName functionality

diff --git a/Engine/Plugins/Online/OnlineFramework/OnlineFramework.uplugin b/Engine/Plugins/Online/OnlineFramework/OnlineFramework.uplugin
--- a/Engine/Plugins/Online/OnlineFramework/OnlineFramework.uplugin
+++ b/Engine/Plugins/Online/OnlineFramework/OnlineFramework.uplugin
@@ -41,7 +41,7 @@
 		},
 		{
 			"Name": "LoginFlow",
-			"Type":  "ClientOnly",
+			"Type": "ClientOnly",
 			"PlatformDenyList": [
 				"Mac"
 			]
@@ -54,6 +54,10 @@
 		}
 	],
 	"Plugins": [
+		{
+			"Name": "OnlineFrameworkCommon",
+			"Enabled": true
+		},
 		{
 			"Name": "OnlineSubsystem",
 			"Enabled": true
@@ -61,6 +65,10 @@
 		{
 			"Name": "OnlineSubsystemUtils",
 			"Enabled": true
+		},
+		{
+			"Name": "OnlineServicesOSSAdapter",
+			"Enabled": true
 		}
 	]
 }
diff --git a/Engine/Plugins/Online/OnlineFramework/Source/Party/Party.Build.cs b/Engine/Plugins/Online/OnlineFramework/Source/Party/Party.Build.cs
--- a/Engine/Plugins/Online/OnlineFramework/Source/Party/Party.Build.cs
+++ b/Engine/Plugins/Online/OnlineFramework/Source/Party/Party.Build.cs
@@ -25,6 +25,9 @@ public class Party : ModuleRules
 				"Core",
 				"CoreUObject",
 				"Engine",
+				"OnlineFrameworkCommon",
+				"OnlineServicesInterface",
+				"OnlineServicesOSSAdapter"
 			}
 			);
 	}
diff --git a/Engine/Plugins/Online/OnlineFramework/Source/Party/Private/SocialToolkit.cpp b/Engine/Plugins/Online/OnlineFramework/Source/Party/Private/SocialToolkit.cpp
--- a/Engine/Plugins/Online/OnlineFramework/Source/Party/Private/SocialToolkit.cpp
+++ b/Engine/Plugins/Online/OnlineFramework/Source/Party/Private/SocialToolkit.cpp
@@ -10,6 +10,12 @@
 #include "OnlineSubsystemUtils.h"
 #include "Engine/LocalPlayer.h"
 
+#include "Online/OnlineServices.h"
+#include "Online/PresenceOSSAdapter.h"
+#include "OnlineFramework/CommonAccount.h"
+#include "OnlineFramework/CommonAccountUtils.h"
+#include "OnlineFramework/CommonConfig.h"
+
 #include UE_INLINE_GENERATED_CPP_BY_NAME(SocialToolkit)
 
 #if WITH_EDITOR
@@ -279,6 +285,34 @@ const FOnlineUserPresence* USocialToolkit::GetPresenceInfo(ESocialSubsystem Subs
 	return nullptr;
 }
 
+const UE::Online::TOnlineResult<UE::Online::FGetCachedPresence> USocialToolkit::GetPresenceInfoV2(ESocialSubsystem SubsystemType) const
+{
+	using namespace UE::Online;
+	using namespace UE::OnlineFramework;
+
+	TOnlineResult<FGetCachedPresence> Result(Errors::NotFound());
+
+	FCommonConfig CommonConfig(this);
+	if (IOnlineServicesPtr OnlineServices = CommonConfig.GetServices(PresenceFrameworkInstance))
+	{
+		FUniqueNetIdRepl LocalUserNetId = GetLocalUserNetId(SubsystemType);
+		if(LocalUserNetId.IsValid())
+		{
+			if (const FCommonAccountPtr LocalUserCommonAccountPtr = GetCommonAccountFromV1(CommonConfig, LocalUserNetId->AsShared(), GetFirstFrameworkInstanceName(LocalUserNetId)))
+			{
+				FAccountId LocalUserAccountId = LocalUserCommonAccountPtr->GetId(PresenceFrameworkInstance);
+
+				if (IPresencePtr Presence = OnlineServices->GetPresenceInterface())
+				{
+					Result = Presence->GetCachedPresence({ LocalUserAccountId, LocalUserAccountId });
+				}
+			}
+		}
+	}
+
+	return Result;
+}
+
 void USocialToolkit::SetLocalUserOnlineState(EOnlinePresenceState::Type OnlineState)
 {
 	if (IOnlineSubsystem* PrimaryOss = GetSocialOss(ESocialSubsystem::Primary))
@@ -287,9 +321,6 @@ void USocialToolkit::SetLocalUserOnlineState(EOnlinePresenceState::Type OnlineSt
 		 FUniqueNetIdRepl LocalUserId = GetLocalUserNetId(ESocialSubsystem::Primary);
 		 if (PresenceInterface.IsValid() && LocalUserId.IsValid())
 		 {
-			 TSharedPtr<FOnlineUserPresence> CurrentPresence;
-			 PresenceInterface->GetCachedPresence(*LocalUserId, CurrentPresence);
-
 			 FOnlinePresenceSetPresenceParameters NewStatus;
 			 NewStatus.State = OnlineState;
 			 PresenceInterface->SetPresence(*LocalUserId, MoveTemp(NewStatus));
@@ -297,6 +328,37 @@ void USocialToolkit::SetLocalUserOnlineState(EOnlinePresenceState::Type OnlineSt
 	}
 }
 
+void USocialToolkit::SetLocalUserOnlineStateV2(UE::Online::EUserPresenceStatus OnlineState)
+{
+	using namespace UE::Online;
+	using namespace UE::OnlineFramework;
+	
+	FCommonConfig CommonConfig(this);
+	if (IOnlineServicesPtr OnlineServices = CommonConfig.GetServices(PresenceFrameworkInstance))
+	{
+		FUniqueNetIdRepl LocalUserNetId = GetLocalUserNetId(ESocialSubsystem::Primary);
+		if(LocalUserNetId.IsValid())
+		{
+			if (const FCommonAccountPtr LocalUserCommonAccountPtr = GetCommonAccountFromV1(CommonConfig, LocalUserNetId->AsShared(), GetFirstFrameworkInstanceName(LocalUserNetId)))
+			{
+				FAccountId LocalUserAccountId = LocalUserCommonAccountPtr->GetId(PresenceFrameworkInstance);
+
+				if (IPresencePtr Presence = OnlineServices->GetPresenceInterface())
+				{
+					FPartialUpdatePresence::Params::FMutations Mutations;
+					Mutations.Status = OnlineState;
+
+					FPartialUpdatePresence::Params Params = { LocalUserAccountId, Mutations };
+
+					UE_LOG(LogParty, Verbose, TEXT("Updating rich presence with OnlineServices"));
+
+					Presence->PartialUpdatePresence(MoveTemp(Params));
+				}
+			}
+		}
+	}
+}
+
 void USocialToolkit::AddLocalUserOnlineProperties(FPresenceProperties OnlineProperties)
 {
 	if (IOnlineSubsystem* PrimaryOss = GetSocialOss(ESocialSubsystem::Primary))
@@ -589,6 +651,16 @@ bool USocialToolkit::IsFriendshipRestricted(const USocialUser& SocialUser, ESoci
 	return false;
 }
 
+bool USocialToolkit::IsOnlineServicesPresenceEnabled() const
+{
+	if (IConsoleVariable* CVar = IConsoleManager::Get().FindConsoleVariable(TEXT("onlineframework.party.isonlineservicespresenceenabled")))
+	{
+		return CVar->GetBool();
+	}
+
+	return false;
+}
+
 //@todo DanH: Rename this in a way that keeps the intent but relates that more than just the primary login has completed (i.e. the game has also completed whatever specific stuff it wants to for login as well)
 void USocialToolkit::OnOwnerLoggedIn()
 {
@@ -633,6 +705,18 @@ void USocialToolkit::OnOwnerLoggedIn()
 			{
 				PresenceInterface->AddOnPresenceReceivedDelegate_Handle(FOnPresenceReceivedDelegate::CreateUObject(this, &USocialToolkit::HandlePresenceReceived, SubsystemType));
 			}
+
+			if (IsOnlineServicesPresenceEnabled())
+			{
+				UE::OnlineFramework::FCommonConfig CommonConfig(this);
+				if (UE::Online::IOnlineServicesPtr OnlineServices = CommonConfig.GetServices(PresenceFrameworkInstance))
+				{
+					if (UE::Online::IPresencePtr Presence = OnlineServices->GetPresenceInterface())
+					{
+						OnPresenceUpdatedHandle = Presence->OnPresenceUpdated().Add(this, &USocialToolkit::HandleServicesPresenceUpdated, SubsystemType);
+					}
+				}
+			}
 		}
 	}
 
@@ -696,6 +780,8 @@ void USocialToolkit::OnOwnerLoggedOut()
 			{
 				PresenceInterface->ClearOnPresenceReceivedDelegates(this);
 			}
+
+			OnPresenceUpdatedHandle.Unbind();
 		}
 	}
 
@@ -913,6 +999,42 @@ void USocialToolkit::HandleMapExternalIdComplete(ESocialSubsystem SubsystemType,
 	}
 }
 
+void USocialToolkit::HandleServicesPresenceUpdated(const UE::Online::FPresenceUpdated& Event, ESocialSubsystem SubsystemType)
+{
+	using namespace UE::Online;
+	using namespace UE::OnlineFramework;
+
+	TSharedRef<const FUserPresence> UpdatedPresence = Event.UpdatedPresence;
+
+	FCommonConfig CommonConfig(this);
+	
+	// We need a services instance with an Auth adapter to make this conversion
+	if (FCommonAccountPtr CommonAccountPtr = GetCommonAccountFromV2(CommonConfig, UpdatedPresence->AccountId, PresenceFrameworkAdapterInstance))
+	{
+		if (FUniqueNetIdPtr UserId = GetV1FromCommonAccount(CommonAccountPtr.ToSharedRef(), PresenceFrameworkAdapterInstance))
+		{
+			QUICK_SCOPE_CYCLE_COUNTER(STAT_USocialToolkit_HandlePresenceReceived);
+			if (USocialUser* UpdatedUser = FindUser(UserId.ToSharedRef()))
+			{
+				UpdatedUser->NotifyPresenceChanged(SubsystemType);
+
+				if (UpdatedUser->IsFriend())
+				{
+					QueueUserDependentActionInternal(UserId.ToSharedRef(), ESocialSubsystem::Primary,
+						[this, UpdatedPresence](USocialUser& SocialUser)
+						{
+							OnFriendPresenceDidChangeV2(SocialUser, UpdatedPresence, ESocialSubsystem::Primary);
+						});
+				}
+			}
+			else if (SubsystemType == ESocialSubsystem::Platform)
+			{
+				UE_LOG(LogParty, Error, TEXT("Platform presence received for UserId [%s], but existing SocialUser could not be found."), *UserId->ToDebugString());
+			}
+		}
+	}
+}
+
 void USocialToolkit::HandlePresenceReceived(const FUniqueNetId& UserId, const TSharedRef<FOnlineUserPresence>& NewPresence, ESocialSubsystem SubsystemType)
 {
 	QUICK_SCOPE_CYCLE_COUNTER(STAT_USocialToolkit_HandlePresenceReceived);
diff --git a/Engine/Plugins/Online/OnlineFramework/Source/Party/Private/User/SocialUser.cpp b/Engine/Plugins/Online/OnlineFramework/Source/Party/Private/User/SocialUser.cpp
--- a/Engine/Plugins/Online/OnlineFramework/Source/Party/Private/User/SocialUser.cpp
+++ b/Engine/Plugins/Online/OnlineFramework/Source/Party/Private/User/SocialUser.cpp
@@ -3,15 +3,19 @@
 #include "User/SocialUser.h"
 
 #include "Interactions/SocialInteractionHandle.h"
-#include "SocialQuery.h"
-#include "SocialManager.h"
-#include "Party/SocialParty.h"
-#include "Party/PartyMember.h"
-
+#include "Interfaces/OnlineUserInterface.h"
 #include "Misc/CommandLine.h"
 
+#include "Online/OnlineErrorDefinitions.h"
+#include "Online/OnlineServices.h"
+#include "OnlineFramework/CommonAccount.h"
+#include "OnlineFramework/CommonAccountUtils.h"
+#include "OnlineFramework/CommonConfig.h"
 #include "OnlineSubsystemUtils.h"
-#include "Interfaces/OnlineUserInterface.h"
+#include "Party/PartyMember.h"
+#include "Party/SocialParty.h"
+#include "SocialManager.h"
+#include "SocialQuery.h"
 
 #include UE_INLINE_GENERATED_CPP_BY_NAME(SocialUser)
 
@@ -113,6 +117,82 @@ TArray<TWeakObjectPtr<USocialUser>> USocialUser::FDebugInitializer::PausedSocial
 // FSocialQuery_UserInfo
 //////////////////////////////////////////////////////////////////////////
 
+class FSocialQuery_UserInfoV2 : public TSocialQuery<UE::Online::FAccountId, const TValueOrError<TSharedRef<UE::Online::FUserInfo>, UE::Online::FOnlineError>&>
+{
+public:
+	friend class USocialToolkit;
+
+	static FName GetQueryId() { return TEXT("UserInfoV2"); }
+
+	virtual void ExecuteQuery() override
+	{
+		using namespace UE::Online;
+		using namespace UE::OnlineFramework;
+
+		if (const USocialToolkit* StrongToolkit = Toolkit.Get())
+		{
+			FCommonConfig CommonConfig(StrongToolkit);
+			if (IOnlineServicesPtr Services = CommonConfig.GetServices(StrongToolkit->GetPresenceFrameworkInstance()))
+			{
+				if (IUserInfoPtr UserInfo = Services->GetUserInfoInterface())
+				{
+					bHasExecuted = true;
+
+					const FUniqueNetIdRepl& LocalUserNetId = Toolkit->GetLocalUserNetId(SubsystemType);
+					if (const FCommonAccountPtr LocalUserCommonAccountPtr = GetCommonAccountFromV1(CommonConfig, LocalUserNetId->AsShared(), GetFirstFrameworkInstanceName(LocalUserNetId)))
+					{
+						FAccountId LocalUserAccountId = LocalUserCommonAccountPtr->GetId(StrongToolkit->GetPresenceFrameworkInstance());
+
+						TArray<FAccountId> SocialUserAccountIds;
+						CompletionCallbacksByUserId.GenerateKeyArray(SocialUserAccountIds);
+						UE_LOG(LogParty, Log, TEXT("FSocialQuery_UserInfoV2 executing for [%d] users on subsystem [%s]"), SocialUserAccountIds.Num(), ToString(SubsystemType));
+
+						UserInfo->QueryUserInfo({ LocalUserAccountId, SocialUserAccountIds })
+							.OnComplete([this, LocalUserAccountId, SocialUserAccountIds, UserInfoWeakPtr = UserInfo.ToWeakPtr()](const TOnlineResult<FQueryUserInfo>& Result)
+							{
+								if (Result.IsOk())
+								{
+									UE_LOG(LogParty, Log, TEXT("FSocialQuery_UserInfoV2 completed query on services [%s]"), ToString(SubsystemType));
+											
+									using namespace UE::Online;
+
+									if (IUserInfoPtr UserInfo = UserInfoWeakPtr.Pin())
+									{
+										// Notify users of the query completion
+										for (const FAccountId& UserAccountId : SocialUserAccountIds)
+										{
+											TOnlineResult<FGetUserInfo> GetUserInfoResult = UserInfo->GetUserInfo({ LocalUserAccountId, UserAccountId });
+											if (GetUserInfoResult.IsOk())
+											{
+												if (const FOnQueryComplete* CompletionCallback = CompletionCallbacksByUserId.Find(UserAccountId))
+												{
+													CompletionCallback->ExecuteIfBound(SubsystemType, true, MakeValue(GetUserInfoResult.GetOkValue().UserInfo));
+												}
+											}
+											else
+											{
+												if (const FOnQueryComplete* CompletionCallback = CompletionCallbacksByUserId.Find(UserAccountId))
+												{
+													CompletionCallback->ExecuteIfBound(SubsystemType, false, MakeError(GetUserInfoResult.GetErrorValue()));
+												}
+											}
+										}
+									}
+
+									OnQueryCompleted.ExecuteIfBound(GetQueryId(), AsShared());
+								}
+								else
+								{
+									UE_LOG(LogParty, Log, TEXT("FSocialQuery_UserInfoV2 failed query on services [%s] with error [%s]"), ToString(SubsystemType), *Result.GetErrorValue().GetLogString());
+								}
+							});
+					}
+				}
+			}
+		}
+	}
+};
+
 class FSocialQuery_UserInfo : public TSocialQuery<FUniqueNetIdRef, const TSharedPtr<FOnlineUser>&>
 {
 public:
@@ -422,6 +502,76 @@ EOnlinePresenceState::Type USocialUser::GetOnlineStatus() const
 	return OnlineStatus;
 }
 
+UE::Online::EUserPresenceStatus USocialUser::GetOnlineStatusV2() const
+{
+	using namespace UE::Online;
+	using namespace UE::OnlineFramework;
+
+	EUserPresenceStatus Result = EUserPresenceStatus::Unknown;
+
+	if (IsLocalUser())
+	{
+		TOnlineResult<FGetCachedPresence> GetCachedPresenceResult = GetOwningToolkit().GetPresenceInfoV2(ESocialSubsystem::Primary);
+		if (GetCachedPresenceResult.IsOk())
+		{
+			Result = GetCachedPresenceResult.GetOkValue().Presence->Status;
+		}
+
+		return Result;
+	}
+
+#if WITH_EDITOR
+	if (bDebug_IsPresenceArtificial && !IsBlocked())
+	{
+		return Debug_RandomPresenceV2;
+	}
+#endif
+
+	Result = EUserPresenceStatus::Offline;
+
+	FCommonConfig CommonConfig(this);
+	if (IOnlineServicesPtr OnlineServices = CommonConfig.GetServices(PresenceFrameworkInstance))
+	{
+		if (IPresencePtr Presence = OnlineServices->GetPresenceInterface())
+		{
+			// Get the most "present" status available on any of the associated platforms
+			for (const TPair<ESocialSubsystem, FSubsystemUserInfo>& SubsystemInfoPair : SubsystemInfoByType)
+			{
+				const FUniqueNetIdRepl& LocalUserNetId = GetOwningToolkit().GetLocalUserNetId(SubsystemInfoPair.Key);
+				if(LocalUserNetId.IsValid())
+				{
+					if (const FCommonAccountPtr LocalUserCommonAccountPtr = GetCommonAccountFromV1(CommonConfig, LocalUserNetId->AsShared(), GetFirstFrameworkInstanceName(LocalUserNetId)))
+					{
+						FAccountId LocalUserAccountId = LocalUserCommonAccountPtr->GetId(PresenceFrameworkInstance);
+
+						const FUniqueNetIdRepl& UserId = SubsystemInfoPair.Value.GetUserId();
+
+						if (const FCommonAccountPtr CommonAccountPtr = GetCommonAccountFromV1(CommonConfig, UserId->AsShared(), GetFirstFrameworkInstanceName(UserId)))
+						{
+							FAccountId AccountId = CommonAccountPtr->GetId(PresenceFrameworkInstance);
+
+							TOnlineResult<FGetCachedPresence> GetCachedPresenceResult = Presence->GetCachedPresence({ LocalUserAccountId, AccountId });
+							if (GetCachedPresenceResult.IsOk())
+							{
+								EUserPresenceStatus CachedPresenceStatus = GetCachedPresenceResult.GetOkValue().Presence->Status;
+
+								if (Result == EUserPresenceStatus::Offline || CachedPresenceStatus == EUserPresenceStatus::Online ||
+								(CachedPresenceStatus == EUserPresenceStatus::Away && Result != EUserPresenceStatus::Online))
+								{
+									// Either the best we have is offline, or the new one is either online or away (if necessary we can get into the weeds of prioritizing the other states)
+									Result = CachedPresenceStatus;
+								}
+							}
+						}
+					}
+				}
+			}
+		}
+	}
+
+	return Result;
+}
+
 void USocialUser::TryBroadcastInitializationComplete()
 {
 	if (!bIsInitialized && NumPendingQueries == 0)
@@ -489,21 +639,86 @@ USocialUser::FSubsystemUserInfo& USocialUser::FindOrCreateSubsystemInfo(const FU
 	return *SubsystemInfo;
 }
 
+bool USocialUser::IsOnlineServicesPresenceEnabled() const
+{
+	if (IConsoleVariable* CVar = IConsoleManager::Get().FindConsoleVariable(TEXT("onlineframework.party.isonlineservicespresenceenabled")))
+	{
+		return CVar->GetBool();
+	}
+
+	return false;
+}
+
 FString USocialUser::GetDisplayName() const
 {
 	FString DisplayName;
 
-	const FSubsystemUserInfo* SubsystemInfo = SubsystemInfoByType.Find(ESocialSubsystem::Primary);
-	if (TSharedPtr<FOnlineUser> PrimaryUserInfo = SubsystemInfo ? SubsystemInfo->UserInfo.Pin() : nullptr)
+	if(const FSubsystemUserInfo* SubsystemInfo = SubsystemInfoByType.Find(ESocialSubsystem::Primary))
 	{
-		// The primary user info has knowledge of display names on all linked accounts
-		const FUserPlatform& UserCurrentPlatform = GetCurrentPlatform();
-		DisplayName = PrimaryUserInfo->GetDisplayName(UserCurrentPlatform.ToString());
-	}
-	else
-	{
-		// We don't have primary user info (so we're not even initialized yet!), but a good-faith effort is to see if we have a platform name
-		DisplayName = GetDisplayName(ESocialSubsystem::Platform);
+		// Temporary redundancy check
+		FString DisplayNameV1, DisplayNameV2;
+
+		// We always retrieve the V1 name during the transition period, for the redundancy check
+		if (TSharedPtr<FOnlineUser> PrimaryUserInfo = SubsystemInfo->UserInfo.Pin())
+		{
+			// The primary user info has knowledge of display names on all linked accounts
+			const FUserPlatform& UserCurrentPlatform = GetCurrentPlatform();
+			DisplayNameV1 = PrimaryUserInfo->GetDisplayName(UserCurrentPlatform.ToString());
+		}
+		else
+		{
+			// We don't have primary user info (so we're not even initialized yet!), but a good-faith effort is to see if we have a platform name
+			DisplayNameV1 = GetDisplayName(ESocialSubsystem::Platform);
+		}
+
+		DisplayName = DisplayNameV1;
+
+		if (IsOnlineServicesPresenceEnabled())
+		{
+			using namespace UE::Online;
+			using namespace UE::OnlineFramework;
+
+			FCommonConfig CommonConfig(this);
+
+			const FUniqueNetIdRepl& LocalUserNetId = GetOwningToolkit().GetLocalUserNetId(ESocialSubsystem::Primary);
+			if (const FCommonAccountPtr LocalUserCommonAccountPtr = GetCommonAccountFromV1(CommonConfig, LocalUserNetId->AsShared(), GetFirstFrameworkInstanceName(LocalUserNetId)))
+			{
+				FAccountId LocalUserAccountId = LocalUserCommonAccountPtr->GetId(PresenceFrameworkInstance);
+
+				const FUniqueNetIdRepl& SocialUserId = SubsystemInfo->GetUserId();
+
+				if (const FCommonAccountPtr SocialUserCommonAccountPtr = GetCommonAccountFromV1(CommonConfig, SocialUserId->AsShared(), GetFirstFrameworkInstanceName(SocialUserId)))
+				{
+					FAccountId SocialUserAccountId = SocialUserCommonAccountPtr->GetId(PresenceFrameworkInstance);
+
+					if (IOnlineServicesPtr OnlineServices = CommonConfig.GetServices(PresenceFrameworkInstance))
+					{
+						if (IUserInfoPtr UserInfo = OnlineServices->GetUserInfoInterface())
+						{
+							TOnlineResult<FGetUserInfo> GetUserInfoResult = UserInfo->GetUserInfo({ LocalUserAccountId, SocialUserAccountId }); //, GetCurrentPlatform().ToString()}); // TODO: We probably need a translator from GetCurrentPlatform to what EOS understands
+							if (GetUserInfoResult.IsOk())
+							{
+								DisplayNameV2 = GetUserInfoResult.GetOkValue().UserInfo->DisplayName;
+							}
+							else
+							{
+								// We don't have primary user info (so we're not even initialized yet!), but a good-faith effort is to see if we have a platform name
+								DisplayNameV2 = GetDisplayName(ESocialSubsystem::Platform);
+							}
+
+#if !UE_BUILD_SHIPPING
+							if (DisplayNameV1 != DisplayNameV2)
+							{
+								UE_LOG(LogParty, Warning, TEXT("DisplayName disparity: DisplayNameV1: %s. DisplayNameV2: %s"), *DisplayNameV1, *DisplayNameV2);
+							}
+#endif
+
+							DisplayName = DisplayNameV2;
+						}
+					}
+				}
+			}
+		}
 	}
 
 	return DisplayName;
@@ -513,7 +728,45 @@ FString USocialUser::GetDisplayName(ESocialSubsystem SubsystemType) const
 {
 	if (const FSubsystemUserInfo* SubsystemInfo = SubsystemInfoByType.Find(SubsystemType))
 	{
-		return SubsystemInfo->GetDisplayName();
+		if (TSharedPtr<FOnlineUser> PrimaryUserInfo = SubsystemInfo->UserInfo.Pin())
+		{
+			if (IsOnlineServicesPresenceEnabled() && SubsystemType == ESocialSubsystem::Primary) // We don't want V2 functionality in Platform for now
+			{
+				using namespace UE::Online;
+				using namespace UE::OnlineFramework;
+
+				FCommonConfig CommonConfig(this);
+
+				const FUniqueNetIdRepl& LocalUserNetId = GetOwningToolkit().GetLocalUserNetId(ESocialSubsystem::Primary);
+				if (const FCommonAccountPtr LocalUserCommonAccountPtr = GetCommonAccountFromV1(CommonConfig, LocalUserNetId->AsShared(), GetFirstFrameworkInstanceName(LocalUserNetId)))
+				{
+					FAccountId LocalUserAccountId = LocalUserCommonAccountPtr->GetId(PresenceFrameworkInstance);
+
+					const FUniqueNetIdRepl& SocialUserId = SubsystemInfo->GetUserId();
+
+					if (const FCommonAccountPtr SocialUserCommonAccountPtr = GetCommonAccountFromV1(CommonConfig, SocialUserId->AsShared(), GetFirstFrameworkInstanceName(SocialUserId)))
+					{
+						FAccountId SocialUserAccountId = SocialUserCommonAccountPtr->GetId(PresenceFrameworkInstance);
+
+						if (IOnlineServicesPtr OnlineServices = CommonConfig.GetServices(PresenceFrameworkInstance))
+						{
+							if (IUserInfoPtr UserInfo = OnlineServices->GetUserInfoInterface())
+							{
+								TOnlineResult<FGetUserInfo> GetUserInfoResult = UserInfo->GetUserInfo({ LocalUserAccountId, SocialUserAccountId });
+								if (GetUserInfoResult.IsOk())
+								{
+									return GetUserInfoResult.GetOkValue().UserInfo->DisplayName;
+								}
+							}
+						}
+					}
+				}
+			}
+			else
+			{
+				return PrimaryUserInfo->GetDisplayName();
+			}
+		}
 	}
 	return FString();
 }
@@ -655,6 +908,48 @@ const FOnlineUserPresence* USocialUser::GetFriendPresenceInfo(ESocialSubsystem S
 	return nullptr;
 }
 
+UE::Online::TOnlineResult<UE::Online::FGetCachedPresence> USocialUser::GetFriendPresenceInfoV2(ESocialSubsystem SubsystemType) const
+{	
+	using namespace UE::Online;
+	using namespace UE::OnlineFramework;
+
+	TOnlineResult<FGetCachedPresence> Result(Errors::NotFound());
+
+	FAccountId LocalUserAccountId, FriendAccountId;
+
+	FCommonConfig CommonConfig(this);
+	FUniqueNetIdRepl LocalUserNetId = GetOwningToolkit().GetLocalUserNetId(SubsystemType);
+	if (const FCommonAccountPtr LocalUserCommonAccountPtr = GetCommonAccountFromV1(CommonConfig, LocalUserNetId->AsShared(), GetFirstFrameworkInstanceName(LocalUserNetId)))
+	{
+		LocalUserAccountId = LocalUserCommonAccountPtr->GetId(PresenceFrameworkInstance);
+	}
+
+	if(const FSubsystemUserInfo* SubsystemInfo = SubsystemInfoByType.Find(SubsystemType))
+	{
+		FUniqueNetIdRepl FriendUserId = SubsystemInfo->GetUserId();
+
+		if (const FCommonAccountPtr FriendCommonAccountPtr = GetCommonAccountFromV1(CommonConfig, FriendUserId->AsShared(), GetFirstFrameworkInstanceName(FriendUserId)))
+		{
+			FriendAccountId = FriendCommonAccountPtr->GetId(PresenceFrameworkInstance);
+		}
+	}
+
+	if(LocalUserAccountId.IsValid() && FriendAccountId.IsValid())
+	{
+		if (IOnlineServicesPtr OnlineServices = CommonConfig.GetServices(PresenceFrameworkInstance))
+		{
+			if (IPresencePtr Presence = OnlineServices->GetPresenceInterface())
+			{
+				FGetCachedPresence::Params Params = { LocalUserAccountId, FriendAccountId };
+
+				Result = Presence->GetCachedPresence(MoveTemp(Params));
+			}
+		}
+	}
+
+	return Result;
+}
+
 FDateTime USocialUser::GetFriendshipCreationDate() const
 {
 	if (const FSubsystemUserInfo* SubsystemInfo = SubsystemInfoByType.Find(ESocialSubsystem::Primary))
@@ -1435,6 +1730,7 @@ void USocialUser::Debug_RandomizePresence()
 {
 	bDebug_IsPresenceArtificial = true;
 	Debug_RandomPresence = static_cast<EOnlinePresenceState::Type>(FMath::RandRange((int32)EOnlinePresenceState::Online, (int32)EOnlinePresenceState::Away));
+	Debug_RandomPresenceV2 = static_cast<UE::Online::EUserPresenceStatus>(FMath::RandRange((int32)UE::Online::EUserPresenceStatus::Offline, (int32)UE::Online::EUserPresenceStatus::Unknown));
 }
 #endif
 
@@ -1546,6 +1842,11 @@ void USocialUser::SetUserInfo(ESocialSubsystem SubsystemType, const TSharedRef<F
 	}
 }
 
+TFuture<void> USocialUser::SetupCommonAccount()
+{
+	return MakeFulfilledPromise<void>().GetFuture();
+}
+
 void USocialUser::HandleQueryUserInfoComplete(ESocialSubsystem SubsystemType, bool bWasSuccessful, const TSharedPtr<FOnlineUser>& UserInfo)
 {
 	--NumPendingQueries;
@@ -1555,10 +1856,53 @@ void USocialUser::HandleQueryUserInfoComplete(ESocialSubsystem SubsystemType, bo
 		SetUserInfo(SubsystemType, UserInfo.ToSharedRef());
 	}
 
+	if (IsOnlineServicesPresenceEnabled())
+	{
+		SetupCommonAccount()
+		.Then([WeakThis = TWeakObjectPtr<USocialUser>(this), SubsystemType](TFuture<void> Result)
+		{
+			if (USocialUser* StrongThis = WeakThis.Get())
+			{
+				using namespace UE::OnlineFramework;
+
+				UE_LOG(LogParty, VeryVerbose, TEXT("SocialUser [%s] querying user info on services [%s]"), *StrongThis->ToDebugString(), ToString(SubsystemType));
+
+				const FUniqueNetIdRepl UserId = StrongThis->GetUserId(SubsystemType);
+
+				FCommonConfig CommonConfig(StrongThis);
+				if (const FCommonAccountPtr CommonAccountPtr = GetCommonAccountFromV1(CommonConfig, UserId->AsShared(), GetFirstFrameworkInstanceName(UserId)))
+				{
+					UE::Online::FAccountId AccountId = CommonAccountPtr->GetId(StrongThis->PresenceFrameworkInstance);
+					if (AccountId.IsValid())
+					{
+						// No valid user info exists on this subsystem, so queue up a query for it
+						FSocialQuery_UserInfoV2::FOnQueryComplete QueryCompleteHandler = FSocialQuery_UserInfoV2::FOnQueryComplete::CreateUObject(StrongThis, &USocialUser::HandleQueryUserInfoV2Complete);
+						FSocialQueryManager::AddUserId<FSocialQuery_UserInfoV2>(StrongThis->GetOwningToolkit(), SubsystemType, AccountId, QueryCompleteHandler);
+						StrongThis->NumPendingQueries++;
+
+						return;
+					}
+				}
+
+				StrongThis->TryBroadcastInitializationComplete();
+			}
+		});
+		
+		return;	
+	}
+
 	UE_LOG(LogParty, VeryVerbose, TEXT("User [%s] finished querying user info on subsystem [%s] with result [%d]. [%d] queries still pending."), *ToDebugString(), ToString(SubsystemType), UserInfo.IsValid(), NumPendingQueries);
 	TryBroadcastInitializationComplete();
 }
 
+void USocialUser::HandleQueryUserInfoV2Complete(ESocialSubsystem SubsystemType, bool bWasSuccessful, const TValueOrError<TSharedRef<UE::Online::FUserInfo>, UE::Online::FOnlineError>& Result)
+{
+	--NumPendingQueries;
+
+	UE_LOG(LogParty, VeryVerbose, TEXT("User [%s] finished querying user info on subsystem [%s] with result [%d]. [%d] queries still pending."), *ToDebugString(), ToString(SubsystemType), bWasSuccessful, NumPendingQueries);
+	TryBroadcastInitializationComplete();
+}
+
 void USocialUser::HandleSetNicknameComplete(int32 LocalUserNum, const FUniqueNetId& FriendId, const FString& ListName, const FOnlineError& Error)
 {
 	if (!Error.WasSuccessful())
diff --git a/Engine/Plugins/Online/OnlineFramework/Source/Party/Public/SocialToolkit.h b/Engine/Plugins/Online/OnlineFramework/Source/Party/Public/SocialToolkit.h
--- a/Engine/Plugins/Online/OnlineFramework/Source/Party/Public/SocialToolkit.h
+++ b/Engine/Plugins/Online/OnlineFramework/Source/Party/Public/SocialToolkit.h
@@ -6,6 +6,8 @@
 #include "Interfaces/OnlinePresenceInterface.h"
 #include "Containers/Ticker.h"
 
+#include "Online/Presence.h"
+#include "Online/OnlineAsyncOpHandle.h"
 #include "PartyModule.h"
 #include "Templates/SubclassOf.h"
 #include "User/SocialUser.h"
@@ -77,7 +79,9 @@ public:
 	UE_API int32 GetLocalUserNum() const;
 
 	UE_API const FOnlineUserPresence* GetPresenceInfo(ESocialSubsystem SubsystemType) const;
+	UE_API const UE::Online::TOnlineResult<UE::Online::FGetCachedPresence> GetPresenceInfoV2(ESocialSubsystem SubsystemType) const;
 	UE_API void SetLocalUserOnlineState(EOnlinePresenceState::Type OnlineState);
+	UE_API void SetLocalUserOnlineStateV2(UE::Online::EUserPresenceStatus OnlineState);
 
 	UE_DEPRECATED(5.6, "AddLocalUserOnlineProperties is deprecated, please use IOnlinePresence::SetPresence instead.")
 	UE_API void AddLocalUserOnlineProperties(FPresenceProperties OnlineProperties);
@@ -163,6 +167,11 @@ public:
 	bool Debug_IsRandomlyChangingPresence() const { return bDebug_IsRandomlyChangingUserPresence; }
 #endif
 
+	UE_API FName GetPresenceFrameworkInstance() const
+	{
+		return PresenceFrameworkInstance;
+	}
+
 protected:
 	UE_API void NotifySubsystemIdEstablished(USocialUser& SocialUser, ESocialSubsystem SubsystemType, const FUniqueNetIdRepl& SubsystemId);
 	TSubclassOf<USocialChatManager> GetChatManagerClass() { return ChatManagerClass; }
@@ -198,6 +207,7 @@ protected:
 protected:
 	/** Called when a Friend's presence did change */
 	virtual void OnFriendPresenceDidChange(const USocialUser& FriendSocialUser, const TSharedRef<FOnlineUserPresence>& NewPresence, ESocialSubsystem SubsystemType) {}
+	virtual void OnFriendPresenceDidChangeV2(const USocialUser& FriendSocialUser, const TSharedRef<const UE::Online::FUserPresence>& NewPresence, ESocialSubsystem SubsystemType) {}
 
 	UE_API void QueryFriendsLists();
 	UE_API void QueryBlockedPlayers();
@@ -267,6 +277,7 @@ private:	// Handlers
 
 	UE_API void HandleMapExternalIdComplete(ESocialSubsystem SubsystemType, bool bWasSuccessful, const FUniqueNetIdRepl& MappedPrimaryId, FUniqueNetIdRepl ExternalId, TFunction<void(USocialUser&)> UserActionFunc, bool bExecutePostInit);
 	UE_API void HandlePresenceReceived(const FUniqueNetId& UserId, const TSharedRef<FOnlineUserPresence>& NewPresence, ESocialSubsystem SubsystemType);
+	UE_API void HandleServicesPresenceUpdated(const UE::Online::FPresenceUpdated& Event, ESocialSubsystem SubsystemType);
 
 	UE_API void HandleFriendInviteReceived(const FUniqueNetId& LocalUserId, const FUniqueNetId& SenderId, ESocialSubsystem SubsystemType);
 	UE_API void HandleFriendInviteAccepted(const FUniqueNetId& LocalUserId, const FUniqueNetId& NewFriendId, ESocialSubsystem SubsystemType);
@@ -306,11 +317,21 @@ private:	// Handlers
 	FTSTicker::FDelegateHandle Debug_PresenceTickerHandle;
 #endif
 
+	UE::Online::FOnlineEventDelegateHandle OnPresenceUpdatedHandle;
+
 protected:
 	UPROPERTY()
 	TObjectPtr<USocialChatManager> SocialChatManager;
 
 private:
+	bool IsOnlineServicesPresenceEnabled() const;
+
+	UPROPERTY(Config)
+	FName PresenceFrameworkInstance;
+
+	UPROPERTY(Config)
+	FName PresenceFrameworkAdapterInstance;
+
 	static UE_API USocialToolkit* GetToolkitForPlayerInternal(const ULocalPlayer* LocalPlayer);
 	static UE_API TMap<TWeakObjectPtr<const ULocalPlayer>, TWeakObjectPtr<USocialToolkit>> AllToolkitsByOwningPlayer;
 
diff --git a/Engine/Plugins/Online/OnlineFramework/Source/Party/Public/User/SocialUser.h b/Engine/Plugins/Online/OnlineFramework/Source/Party/Public/User/SocialUser.h
--- a/Engine/Plugins/Online/OnlineFramework/Source/Party/Public/User/SocialUser.h
+++ b/Engine/Plugins/Online/OnlineFramework/Source/Party/Public/User/SocialUser.h
@@ -2,9 +2,14 @@
 
 #pragma once
 
+#include "Async/Future.h"
+#include "Online/OnlineResult.h"
+#include "Online/Presence.h"
+#include "Online/UserInfo.h"
 #include "OnlineSubsystemTypes.h"
 #include "Party/PartyTypes.h"
 #include "SocialTypes.h"
+#include "Templates/ValueOrError.h"
 
 #include "SocialUser.generated.h"
 
@@ -53,6 +58,7 @@ public:
 	UE_API bool HasNetId(const FUniqueNetIdRepl& UniqueId) const;
 	UE_API USocialToolkit& GetOwningToolkit() const;
 	UE_API EOnlinePresenceState::Type GetOnlineStatus() const;
+	UE_API UE::Online::EUserPresenceStatus GetOnlineStatusV2() const;
 
 	UE_API FUniqueNetIdRepl GetUserId(ESocialSubsystem SubsystemType) const;
 	UE_API FString GetDisplayName() const;
@@ -66,6 +72,7 @@ public:
 	UE_API bool IsFriendshipPending(ESocialSubsystem SubsystemType) const;
 	UE_API bool IsAnyInboundFriendshipPending() const;
 	UE_API const FOnlineUserPresence* GetFriendPresenceInfo(ESocialSubsystem SubsystemType) const;
+	UE_API UE::Online::TOnlineResult<UE::Online::FGetCachedPresence> GetFriendPresenceInfoV2(ESocialSubsystem SubsystemType) const;
 	UE_API FDateTime GetFriendshipCreationDate() const;
 	UE_API virtual FDateTime GetLastOnlineDate() const;
 	UE_API FText GetSocialName() const;
@@ -192,6 +199,7 @@ protected:
 	UE_API void Debug_RandomizePresence();
 	bool bDebug_IsPresenceArtificial = false;
 	EOnlinePresenceState::Type Debug_RandomPresence;
+	UE::Online::EUserPresenceStatus Debug_RandomPresenceV2;
 #endif
 
 protected:
@@ -240,12 +248,21 @@ protected:
 
 	UE_API void SetUserInfo(ESocialSubsystem SubsystemType, const TSharedRef<FOnlineUser>& UserInfo);
 
+	// This method should be extended in derived classes to asynchronously set up the correct CommonAccount state for the given social user.
+	UE_API virtual TFuture<void> SetupCommonAccount();
+
+	UPROPERTY(Config)
+	FName PresenceFrameworkInstance;
+
 private:
 	UE_API void HandleQueryUserInfoComplete(ESocialSubsystem SubsystemType, bool bWasSuccessful, const TSharedPtr<FOnlineUser>& UserInfo);
+	UE_API void HandleQueryUserInfoV2Complete(ESocialSubsystem SubsystemType, bool bWasSuccessful, const TValueOrError<TSharedRef<UE::Online::FUserInfo>, UE::Online::FOnlineError>& Result);
 
 	UE_API virtual FString SanitizePresenceString(FString InString) const;
 
 private:
+	bool IsOnlineServicesPresenceEnabled() const;
+
 	UE_API FSubsystemUserInfo& FindOrCreateSubsystemInfo(const FUniqueNetIdRepl& SubsystemId, ESocialSubsystem SubsystemType);
 
 	bool bIsInitialized = false;
