Refactored presence prototype to match progressive rollout strategy

diff --git a/Engine/Plugins/Online/OnlineFramework/Source/Party/Private/PartyModule.cpp b/Engine/Plugins/Online/OnlineFramework/Source/Party/Private/PartyModule.cpp
--- a/Engine/Plugins/Online/OnlineFramework/Source/Party/Private/PartyModule.cpp
+++ b/Engine/Plugins/Online/OnlineFramework/Source/Party/Private/PartyModule.cpp
@@ -22,10 +22,26 @@ void FPartyModule::ShutdownModule()
 
 #if !UE_BUILD_SHIPPING
 	/**
-		* CVar to enable usage of Online Services Presence
-		*/
-	static TAutoConsoleVariable<int32> CVar_OnlineFramework_Party_IsOnlineServicesPresenceEnabled(
-		TEXT("onlineframework.party.isonlineservicespresenceenabled"),
+	* CVar to enable usage of Online Services Presence
+	*/
+	static TAutoConsoleVariable<int32> CVar_OnlineFramework_Party_ShouldWritePresenceToOnlineServices(
+		TEXT("onlineframework.party.shouldwritepresencetoonlineservices"),
 		0,
-		TEXT("Whether relevant systems should use OnlineServices for Presence operations"));
+		TEXT("Whether relevant systems should write Presence information to OnlineServices"));
+
+	/**
+	* CVar to enable usage of Online Services Presence
+	*/
+	static TAutoConsoleVariable<int32> CVar_OnlineFramework_Party_ShouldReadPresenceFromOnlineServices(
+		TEXT("onlineframework.party.shouldreadpresencefromonlineservices"),
+		0,
+		TEXT("Whether relevant systems should read Presence information from OnlineServices"));
+
+	/**
+	* CVar to enable usage of Online Services Presence
+	*/
+	static TAutoConsoleVariable<int32> CVar_OnlineFramework_Party_ShouldReadUserInfoFromOnlineServices(
+		TEXT("onlineframework.party.shouldreaduserinfofromonlineservices"),
+		0,
+		TEXT("Whether relevant systems should read User Information from OnlineServices"));
 #endif
\ No newline at end of file
diff --git a/Engine/Plugins/Online/OnlineFramework/Source/Party/Private/SocialToolkit.cpp b/Engine/Plugins/Online/OnlineFramework/Source/Party/Private/SocialToolkit.cpp
--- a/Engine/Plugins/Online/OnlineFramework/Source/Party/Private/SocialToolkit.cpp
+++ b/Engine/Plugins/Online/OnlineFramework/Source/Party/Private/SocialToolkit.cpp
@@ -651,9 +651,9 @@ bool USocialToolkit::IsFriendshipRestricted(const USocialUser& SocialUser, ESoci
 	return false;
 }
 
-bool USocialToolkit::IsOnlineServicesPresenceEnabled() const
+bool USocialToolkit::ShouldReadPresenceFromOnlineServices() const
 {
-	if (IConsoleVariable* CVar = IConsoleManager::Get().FindConsoleVariable(TEXT("onlineframework.party.isonlineservicespresenceenabled")))
+	if (IConsoleVariable* CVar = IConsoleManager::Get().FindConsoleVariable(TEXT("onlineframework.party.shouldreadpresencefromonlineservices")))
 	{
 		return CVar->GetBool();
 	}
@@ -701,12 +701,7 @@ void USocialToolkit::OnOwnerLoggedIn()
 				PartyInterface->AddOnPartyRequestToJoinRemovedDelegate_Handle(FOnPartyRequestToJoinRemovedDelegate::CreateUObject(this, &USocialToolkit::HandlePartyRequestToJoinRemoved));
 			}
 
-			if (IOnlinePresencePtr PresenceInterface = OSS->GetPresenceInterface())
-			{
-				PresenceInterface->AddOnPresenceReceivedDelegate_Handle(FOnPresenceReceivedDelegate::CreateUObject(this, &USocialToolkit::HandlePresenceReceived, SubsystemType));
-			}
-
-			if (IsOnlineServicesPresenceEnabled())
+			if (ShouldReadPresenceFromOnlineServices())
 			{
 				UE::OnlineFramework::FCommonConfig CommonConfig(this);
 				if (UE::Online::IOnlineServicesPtr OnlineServices = CommonConfig.GetServices(PresenceFrameworkInstance))
@@ -717,6 +712,13 @@ void USocialToolkit::OnOwnerLoggedIn()
 					}
 				}
 			}
+			else
+			{
+				if (IOnlinePresencePtr PresenceInterface = OSS->GetPresenceInterface())
+				{
+					PresenceInterface->AddOnPresenceReceivedDelegate_Handle(FOnPresenceReceivedDelegate::CreateUObject(this, &USocialToolkit::HandlePresenceReceived, SubsystemType));
+				}
+			}
 		}
 	}
 
diff --git a/Engine/Plugins/Online/OnlineFramework/Source/Party/Private/User/SocialUser.cpp b/Engine/Plugins/Online/OnlineFramework/Source/Party/Private/User/SocialUser.cpp
--- a/Engine/Plugins/Online/OnlineFramework/Source/Party/Private/User/SocialUser.cpp
+++ b/Engine/Plugins/Online/OnlineFramework/Source/Party/Private/User/SocialUser.cpp
@@ -642,9 +642,9 @@ USocialUser::FSubsystemUserInfo& USocialUser::FindOrCreateSubsystemInfo(const FU
 	return *SubsystemInfo;
 }
 
-bool USocialUser::IsOnlineServicesPresenceEnabled() const
+bool USocialUser::ShouldReadUserInfoFromOnlineServices() const
 {
-	if (IConsoleVariable* CVar = IConsoleManager::Get().FindConsoleVariable(TEXT("onlineframework.party.isonlineservicespresenceenabled")))
+	if (IConsoleVariable* CVar = IConsoleManager::Get().FindConsoleVariable(TEXT("onlineframework.party.shouldreaduserinfofromonlineservices")))
 	{
 		return CVar->GetBool();
 	}
@@ -671,7 +671,7 @@ FString USocialUser::GetDisplayName() const
 			DisplayName = GetDisplayName(ESocialSubsystem::Platform);
 		}
 
-		if (IsOnlineServicesPresenceEnabled())
+		if (ShouldReadUserInfoFromOnlineServices())
 		{
 			using namespace UE::Online;
 			using namespace UE::OnlineFramework;
@@ -705,7 +705,6 @@ FString USocialUser::GetDisplayName() const
 								// We don't have primary user info (so we're not even initialized yet!), but a good-faith effort is to see if we have a platform name
 								DisplayNameV2 = GetDisplayName(ESocialSubsystem::Platform);
 							}
-
 #if !UE_BUILD_SHIPPING
 							// Temporary redundancy check
 							if (DisplayName != DisplayNameV2)
@@ -735,13 +734,9 @@ FString USocialUser::GetDisplayName(ESocialSubsystem SubsystemType) const
 	{
 		if (TSharedPtr<FOnlineUser> PrimaryUserInfo = SubsystemInfo->UserInfo.Pin())
 		{
-			// Temporary redundancy check
-			FString DisplayNameV1, DisplayNameV2;
-
-			DisplayNameV1 = PrimaryUserInfo->GetDisplayName();
-			DisplayName = DisplayNameV1;
+			DisplayName = PrimaryUserInfo->GetDisplayName();
 
-			if (IsOnlineServicesPresenceEnabled() && SubsystemType == ESocialSubsystem::Primary) // We don't want V2 functionality in Platform for now
+			if (ShouldReadUserInfoFromOnlineServices() && SubsystemType == ESocialSubsystem::Primary) // We don't want V2 functionality in Platform for now
 			{
 				using namespace UE::Online;
 				using namespace UE::OnlineFramework;
@@ -766,12 +761,12 @@ FString USocialUser::GetDisplayName(ESocialSubsystem SubsystemType) const
 								TOnlineResult<FGetUserInfo> GetUserInfoResult = UserInfo->GetUserInfo({ LocalUserAccountId, SocialUserAccountId });
 								if (GetUserInfoResult.IsOk())
 								{
-									DisplayNameV2 = GetUserInfoResult.GetOkValue().UserInfo->DisplayName;
+									FString DisplayNameV2 = GetUserInfoResult.GetOkValue().UserInfo->DisplayName;
 #if !UE_BUILD_SHIPPING
 									// Temporary redundancy check
-									if (DisplayNameV1 != DisplayNameV2)
+									if (DisplayName != DisplayNameV2)
 									{
-										UE_LOG(LogParty, Warning, TEXT("DisplayName disparity: DisplayNameV1: %s. DisplayNameV2: %s"), *DisplayNameV1, *DisplayNameV2);
+										UE_LOG(LogParty, Warning, TEXT("DisplayName disparity: DisplayNameV1: %s. DisplayNameV2: %s"), *DisplayName, *DisplayNameV2);
 									}
 #endif
 									if (!DisplayNameV2.IsEmpty())
@@ -1875,7 +1870,7 @@ void USocialUser::HandleQueryUserInfoComplete(ESocialSubsystem SubsystemType, bo
 		SetUserInfo(SubsystemType, UserInfo.ToSharedRef());
 	}
 
-	if (IsOnlineServicesPresenceEnabled())
+	if (ShouldReadUserInfoFromOnlineServices())
 	{
 		SetupCommonAccount()
 		.Then([WeakThis = TWeakObjectPtr<USocialUser>(this), SubsystemType](TFuture<void> Result)
diff --git a/Engine/Plugins/Online/OnlineFramework/Source/Party/Public/SocialToolkit.h b/Engine/Plugins/Online/OnlineFramework/Source/Party/Public/SocialToolkit.h
--- a/Engine/Plugins/Online/OnlineFramework/Source/Party/Public/SocialToolkit.h
+++ b/Engine/Plugins/Online/OnlineFramework/Source/Party/Public/SocialToolkit.h
@@ -324,7 +324,7 @@ protected:
 	TObjectPtr<USocialChatManager> SocialChatManager;
 
 private:
-	bool IsOnlineServicesPresenceEnabled() const;
+	bool ShouldReadPresenceFromOnlineServices() const;
 
 	UPROPERTY(Config)
 	FName PresenceFrameworkInstance;
diff --git a/Engine/Plugins/Online/OnlineFramework/Source/Party/Public/User/SocialUser.h b/Engine/Plugins/Online/OnlineFramework/Source/Party/Public/User/SocialUser.h
--- a/Engine/Plugins/Online/OnlineFramework/Source/Party/Public/User/SocialUser.h
+++ b/Engine/Plugins/Online/OnlineFramework/Source/Party/Public/User/SocialUser.h
@@ -251,6 +251,8 @@ protected:
 	// This method should be extended in derived classes to asynchronously set up the correct CommonAccount state for the given social user.
 	UE_API virtual TFuture<void> SetupCommonAccount();
 
+	UE_API bool ShouldReadUserInfoFromOnlineServices() const;
+
 	UPROPERTY(Config)
 	FName PresenceFrameworkInstance;
 
@@ -261,7 +263,6 @@ private:
 	UE_API virtual FString SanitizePresenceString(FString InString) const;
 
 private:
-	bool IsOnlineServicesPresenceEnabled() const;
 
 	UE_API FSubsystemUserInfo& FindOrCreateSubsystemInfo(const FUniqueNetIdRepl& SubsystemId, ESocialSubsystem SubsystemType);
 
