Minor bugfixes for 44775368

diff --git a/Engine/Plugins/Online/OnlineFramework/Source/Party/Private/PartyModule.cpp b/Engine/Plugins/Online/OnlineFramework/Source/Party/Private/PartyModule.cpp
--- a/Engine/Plugins/Online/OnlineFramework/Source/Party/Private/PartyModule.cpp
+++ b/Engine/Plugins/Online/OnlineFramework/Source/Party/Private/PartyModule.cpp
@@ -22,7 +22,7 @@ void FPartyModule::ShutdownModule()
 
 #if !UE_BUILD_SHIPPING
 	/**
-	* CVar to enable usage of Online Services Presence
+	* CVar to enable writing information to Online Services Presence
 	*/
 	static TAutoConsoleVariable<int32> CVar_OnlineFramework_Party_ShouldWritePresenceToOnlineServices(
 		TEXT("onlineframework.party.shouldwritepresencetoonlineservices"),
@@ -30,7 +30,7 @@ void FPartyModule::ShutdownModule()
 		TEXT("Whether relevant systems should write Presence information to OnlineServices"));
 
 	/**
-	* CVar to enable usage of Online Services Presence
+	* CVar to enable reading information from Online Services Presence
 	*/
 	static TAutoConsoleVariable<int32> CVar_OnlineFramework_Party_ShouldReadPresenceFromOnlineServices(
 		TEXT("onlineframework.party.shouldreadpresencefromonlineservices"),
@@ -38,7 +38,7 @@ void FPartyModule::ShutdownModule()
 		TEXT("Whether relevant systems should read Presence information from OnlineServices"));
 
 	/**
-	* CVar to enable usage of Online Services Presence
+	* CVar to enable reading information from Online Services User Info
 	*/
 	static TAutoConsoleVariable<int32> CVar_OnlineFramework_Party_ShouldReadUserInfoFromOnlineServices(
 		TEXT("onlineframework.party.shouldreaduserinfofromonlineservices"),
diff --git a/Engine/Plugins/Online/OnlineFramework/Source/Party/Private/SocialToolkit.cpp b/Engine/Plugins/Online/OnlineFramework/Source/Party/Private/SocialToolkit.cpp
--- a/Engine/Plugins/Online/OnlineFramework/Source/Party/Private/SocialToolkit.cpp
+++ b/Engine/Plugins/Online/OnlineFramework/Source/Party/Private/SocialToolkit.cpp
@@ -708,7 +708,8 @@ void USocialToolkit::OnOwnerLoggedIn()
 				{
 					if (UE::Online::IPresencePtr Presence = OnlineServices->GetPresenceInterface())
 					{
-						OnPresenceUpdatedHandle = Presence->OnPresenceUpdated().Add(this, &USocialToolkit::HandleServicesPresenceUpdated, SubsystemType);
+						UE::Online::FOnlineEventDelegateHandle OnPresenceUpdatedHandle = Presence->OnPresenceUpdated().Add(this, &USocialToolkit::HandleServicesPresenceUpdated, SubsystemType);
+						OnPresenceUpdatedHandles.Emplace(SubsystemType, MoveTemp(OnPresenceUpdatedHandle));
 					}
 				}
 			}
@@ -783,7 +784,10 @@ void USocialToolkit::OnOwnerLoggedOut()
 				PresenceInterface->ClearOnPresenceReceivedDelegates(this);
 			}
 
-			OnPresenceUpdatedHandle.Unbind();
+			if (UE::Online::FOnlineEventDelegateHandle* Handle = OnPresenceUpdatedHandles.Find(SubsystemType))
+			{
+				Handle->Unbind();
+			}
 		}
 	}
 
diff --git a/Engine/Plugins/Online/OnlineFramework/Source/Party/Public/SocialToolkit.h b/Engine/Plugins/Online/OnlineFramework/Source/Party/Public/SocialToolkit.h
--- a/Engine/Plugins/Online/OnlineFramework/Source/Party/Public/SocialToolkit.h
+++ b/Engine/Plugins/Online/OnlineFramework/Source/Party/Public/SocialToolkit.h
@@ -317,7 +317,7 @@ private:	// Handlers
 	FTSTicker::FDelegateHandle Debug_PresenceTickerHandle;
 #endif
 
-	UE::Online::FOnlineEventDelegateHandle OnPresenceUpdatedHandle;
+	TMap<ESocialSubsystem, UE::Online::FOnlineEventDelegateHandle> OnPresenceUpdatedHandles;
 
 protected:
 	UPROPERTY()
