[UpdateManager] Change cached response timeout to be a hotfixable CVar with a default of 5 minutes.

diff --git a/Engine/Plugins/Online/OnlineFramework/Source/Hotfix/Private/UpdateManager.cpp b/Engine/Plugins/Online/OnlineFramework/Source/Hotfix/Private/UpdateManager.cpp
--- a/Engine/Plugins/Online/OnlineFramework/Source/Hotfix/Private/UpdateManager.cpp
+++ b/Engine/Plugins/Online/OnlineFramework/Source/Hotfix/Private/UpdateManager.cpp
@@ -17,13 +17,17 @@
 
 #include UE_INLINE_GENERATED_CPP_BY_NAME(UpdateManager)
 
-#define UPDATE_CHECK_SECONDS 30.0
-
 static TAutoConsoleVariable<int32> CVarDebugUpdateManager(
 	TEXT("UI.DebugUpdateCheck"),
 	-1,
 	TEXT("Force switch between update states (-1 is off)"));
 
+static TAutoConsoleVariable<float> CVarUpdateManagerCacheTimeoutSeconds(
+	TEXT("UpdateManager.CacheTimeoutSeconds"),
+	300,
+	TEXT("Timeout in seconds to return a previous update manager response."),
+	ECVF_SaveForNextBoot);
+
 struct FLoadingScreenConfig
 {
 public:
@@ -156,7 +160,7 @@ UUpdateManager::EUpdateStartResult UUpdateManager::StartUpdateCheckInternal(cons
 								 LastResult == EUpdateCompletionStatus::UpdateFailure_HotfixCheck ||
 								 LastResult == EUpdateCompletionStatus::UpdateFailure_NotLoggedIn;
 
-		static double CacheTimer = UPDATE_CHECK_SECONDS;
+		const double CacheTimer = CVarUpdateManagerCacheTimeoutSeconds.GetValueOnGameThread();
 		const double TimeSinceCheck = DeltaTime.GetTotalSeconds();
 		if (bForceCheck || TimeSinceCheck >= CacheTimer)
 		{
@@ -173,7 +177,7 @@ UUpdateManager::EUpdateStartResult UUpdateManager::StartUpdateCheckInternal(cons
 		}
 		else
 		{
-			UE_LOG(LogHotfixManager, Display, TEXT("Returning cached update result %d"), (int32)LastResult);
+			UE_LOG(LogHotfixManager, Display, TEXT("Returning cached update result %d. Cache timeout in %.2f/%.2f seconds."), (int32)LastResult, CacheTimer - TimeSinceCheck, CacheTimer);
 			auto StartDelegate = [this, LastResult]()
 			{
 				CheckComplete(LastResult, false);
