Removed TODO after feature investigation

diff --git a/Engine/Plugins/Online/OnlineFramework/Source/Party/Private/User/SocialUser.cpp b/Engine/Plugins/Online/OnlineFramework/Source/Party/Private/User/SocialUser.cpp
--- a/Engine/Plugins/Online/OnlineFramework/Source/Party/Private/User/SocialUser.cpp
+++ b/Engine/Plugins/Online/OnlineFramework/Source/Party/Private/User/SocialUser.cpp
@@ -821,11 +821,7 @@ FString USocialUser::GetNickname() const
 
 		if (const FSubsystemUserInfo* SubsystemInfo = SubsystemInfoByType.Find(ESocialSubsystem::Primary))
 		{
-			/*if (ShouldReadUserInfoFromOnlineServices())
-			{
-				// TODO: Pending: AddNickname support to OnlineServices' UserInfo
-			}
-			else */if (SubsystemInfo->FriendInfo.IsValid())
+			if (SubsystemInfo->FriendInfo.IsValid())
 			{
 				TSharedPtr<FOnlineFriend> OnlineFriend = SubsystemInfo->FriendInfo.Pin();
 				FString Nickname;
