Allow to handle in subclasses when party is in degraded functionality state

diff --git a/Engine/Plugins/Online/OnlineFramework/Source/Party/Public/Party/SocialParty.h b/Engine/Plugins/Online/OnlineFramework/Source/Party/Public/Party/SocialParty.h
--- a/Engine/Plugins/Online/OnlineFramework/Source/Party/Public/Party/SocialParty.h
+++ b/Engine/Plugins/Online/OnlineFramework/Source/Party/Public/Party/SocialParty.h
@@ -351,8 +351,9 @@ private:
 
 	UE_API UPartyMember* GetMemberInternal(const FUniqueNetIdRepl& MemberId) const;
 
-private:	// Handlers
-	UE_API void HandlePartyStateChanged(const FUniqueNetId& LocalUserId, const FOnlinePartyId& PartyId, EPartyState PartyState, EPartyState PreviousPartyState);
+protected:  // Handlers
+	UE_API virtual void HandlePartyStateChanged(const FUniqueNetId& LocalUserId, const FOnlinePartyId& PartyId, EPartyState PartyState, EPartyState PreviousPartyState);
+private:
 	UE_API void HandlePartyConfigChanged(const FUniqueNetId& LocalUserId, const FOnlinePartyId& PartyId, const FPartyConfiguration& PartyConfig);
 	UE_API void HandleUpdatePartyConfigComplete(const FUniqueNetId& LocalUserId, const FOnlinePartyId& PartyId, EUpdateConfigCompletionResult Result);
 	UE_API void HandlePartyDataReceived(const FUniqueNetId& LocalUserId, const FOnlinePartyId& PartyId, const FName& Namespace, const FOnlinePartyData& PartyData);
