Added logic to enable slow percentage base activation of new presence and user info features

diff --git a/Engine/Plugins/Online/OnlineFramework/Source/Party/Private/PartyModule.cpp b/Engine/Plugins/Online/OnlineFramework/Source/Party/Private/PartyModule.cpp
--- a/Engine/Plugins/Online/OnlineFramework/Source/Party/Private/PartyModule.cpp
+++ b/Engine/Plugins/Online/OnlineFramework/Source/Party/Private/PartyModule.cpp
@@ -20,28 +20,26 @@ void FPartyModule::ShutdownModule()
 {
 }
 
-#if !UE_BUILD_SHIPPING
-	/**
-	* CVar to enable writing information to Online Services Presence
-	*/
-	static TAutoConsoleVariable<int32> CVar_OnlineFramework_Party_ShouldWritePresenceToOnlineServices(
-		TEXT("onlineframework.party.shouldwritepresencetoonlineservices"),
-		0,
-		TEXT("Whether relevant systems should write Presence information to OnlineServices"));
-
-	/**
-	* CVar to enable reading information from Online Services Presence
-	*/
-	static TAutoConsoleVariable<int32> CVar_OnlineFramework_Party_ShouldReadPresenceFromOnlineServices(
-		TEXT("onlineframework.party.shouldreadpresencefromonlineservices"),
-		0,
-		TEXT("Whether relevant systems should read Presence information from OnlineServices"));
-
-	/**
-	* CVar to enable reading information from Online Services User Info
-	*/
-	static TAutoConsoleVariable<int32> CVar_OnlineFramework_Party_ShouldReadUserInfoFromOnlineServices(
-		TEXT("onlineframework.party.shouldreaduserinfofromonlineservices"),
-		0,
-		TEXT("Whether relevant systems should read User Information from OnlineServices"));
-#endif
\ No newline at end of file
+/**
+* CVar to enable writing information to Online Services Presence
+*/
+static TAutoConsoleVariable<int32> CVar_OnlineFramework_Party_ShouldWritePresenceToOnlineServices(
+	TEXT("onlineframework.party.shouldwritepresencetoonlineservices"),
+	0,
+	TEXT("Whether relevant systems should write Presence information to OnlineServices"));
+
+/**
+* CVar to enable reading information from Online Services Presence
+*/
+static TAutoConsoleVariable<int32> CVar_OnlineFramework_Party_ShouldReadPresenceFromOnlineServices(
+	TEXT("onlineframework.party.shouldreadpresencefromonlineservices"),
+	0,
+	TEXT("Whether relevant systems should read Presence information from OnlineServices"));
+
+/**
+* CVar to enable reading information from Online Services User Info
+*/
+static TAutoConsoleVariable<int32> CVar_OnlineFramework_Party_ShouldReadUserInfoFromOnlineServices(
+	TEXT("onlineframework.party.shouldreaduserinfofromonlineservices"),
+	0,
+	TEXT("Whether relevant systems should read User Information from OnlineServices"));
