Added new analytics:
Online.TitleFile.EnumerateFilesComplete
UpdateManager.UpdateCheckComplete
Fix analytics which would not fire correctly with early login due to depdendency on PlayerController:
PatchCheck.PatchCheckComplete
Online.TitleFile.InitComplete
Online.TitleFile.FileRequestComplete

diff --git a/Engine/Plugins/Online/OnlineFramework/Source/Hotfix/Private/OnlineHotfixManager.cpp b/Engine/Plugins/Online/OnlineFramework/Source/Hotfix/Private/OnlineHotfixManager.cpp
--- a/Engine/Plugins/Online/OnlineFramework/Source/Hotfix/Private/OnlineHotfixManager.cpp
+++ b/Engine/Plugins/Online/OnlineFramework/Source/Hotfix/Private/OnlineHotfixManager.cpp
@@ -24,6 +24,8 @@
 #include "Curves/CurveVector.h"
 #include "Curves/CurveLinearColor.h"
 #include "Engine/BlueprintGeneratedClass.h"
+#include "Engine/GameInstance.h"
+#include "Engine/GameViewportClient.h"
 
 #include "AutoRTFM.h"
 #include "Serialization/AsyncLoadingFlushContext.h"
@@ -2108,7 +2110,16 @@ bool UOnlineHotfixManager::ShouldHotfixAsset(const FString& AssetPath) const
 
 UWorld* UOnlineHotfixManager::GetWorld() const
 {
-	return OwnerWorld.IsValid() ? OwnerWorld.Get() : nullptr;
+	// OwnerWorld may be stale as it is set at the time that ::Get is called.
+	if (OwnerWorld.IsValid())
+	{
+		return OwnerWorld.Get();
+	}
+
+	// Try to get the active world from the game instance.
+	UGameViewportClient* ActiveViewport = GEngine ? GEngine->GameViewport.Get() : nullptr;
+	UGameInstance* GameInstance = ActiveViewport ? ActiveViewport->GetGameInstance() : nullptr;
+	return GameInstance ? GameInstance->GetWorld() : nullptr;
 }
 
 void UOnlineHotfixManager::StopTrackingInvalidHotfixedAssets()
diff --git a/Engine/Plugins/Online/OnlineFramework/Source/Hotfix/Private/UpdateManager.cpp b/Engine/Plugins/Online/OnlineFramework/Source/Hotfix/Private/UpdateManager.cpp
--- a/Engine/Plugins/Online/OnlineFramework/Source/Hotfix/Private/UpdateManager.cpp
+++ b/Engine/Plugins/Online/OnlineFramework/Source/Hotfix/Private/UpdateManager.cpp
@@ -94,6 +94,11 @@ void UUpdateManager::SetPending()
 	}
 }
 
+void UUpdateManager::RegisterStatsCollector(const TSharedPtr<IUpdateManagerStatsCollector>& InStatsCollector)
+{
+	StatsCollector = InStatsCollector;
+}
+
 void UUpdateManager::Reset()
 {
 	LastUpdateCheck[0] = LastUpdateCheck[1] = FDateTime(0);
@@ -129,6 +134,13 @@ UUpdateManager::EUpdateStartResult UUpdateManager::StartUpdateCheckInternal(cons
 		UE_LOG(LogHotfixManager, Display, TEXT("Update checks disabled!"));
 		bInitialUpdateFinished = true;
 
+		if (TSharedPtr<IUpdateManagerStatsCollector> StrongStatsCollector = StatsCollector.Pin())
+		{
+			StrongStatsCollector->OnCheckStarted({
+				.ContextName = ContextDefinition.Name,
+				.bEnabled = false });
+		}
+
 		// Move to the pending state until the delayed response can fire, to more closely match non-editor behavior.
 		SetUpdateState(EUpdateState::UpdatePending);
 
@@ -164,6 +176,13 @@ UUpdateManager::EUpdateStartResult UUpdateManager::StartUpdateCheckInternal(cons
 		const double TimeSinceCheck = DeltaTime.GetTotalSeconds();
 		if (bForceCheck || TimeSinceCheck >= CacheTimer)
 		{
+			if (TSharedPtr<IUpdateManagerStatsCollector> StrongStatsCollector = StatsCollector.Pin())
+			{
+				StrongStatsCollector->OnCheckStarted({
+					.ContextName = ContextDefinition.Name,
+					.bCheckAvailabilityOnly = ContextDefinition.bCheckAvailabilityOnly });
+			}
+
 			auto StartDelegate = [this]()
 			{
 				// Check for a patch first, then hotfix application
@@ -177,6 +196,15 @@ UUpdateManager::EUpdateStartResult UUpdateManager::StartUpdateCheckInternal(cons
 		}
 		else
 		{
+			if (TSharedPtr<IUpdateManagerStatsCollector> StrongStatsCollector = StatsCollector.Pin())
+			{
+				StrongStatsCollector->OnCheckStarted({
+					.ContextName = ContextDefinition.Name,
+					.bCheckAvailabilityOnly = ContextDefinition.bCheckAvailabilityOnly,
+					.bCachedResponse = true,
+					.CacheExpiration = CacheTimer - TimeSinceCheck });
+			}
+
 			UE_LOG(LogHotfixManager, Display, TEXT("Returning cached update result %d. Cache timeout in %.2f/%.2f seconds."), (int32)LastResult, CacheTimer - TimeSinceCheck, CacheTimer);
 			auto StartDelegate = [this, LastResult]()
 			{
@@ -227,6 +255,11 @@ void UUpdateManager::CheckComplete(EUpdateCompletionStatus Result, bool bUpdateT
 	auto CompletionDelegate = [this, Result, bUpdateTimestamp]()
 	{
 		UE_LOG(LogHotfixManager, Display, TEXT("External CheckComplete %s"), UpdateCompletionEnum ? *UpdateCompletionEnum->GetNameStringByValue((int64)Result) : TEXT("Invalid"));
+		if (TSharedPtr<IUpdateManagerStatsCollector> StrongStatsCollector = StatsCollector.Pin())
+		{
+			StrongStatsCollector->OnCheckComplete(Result);
+		}
+
 		if (!bInitialUpdateFinished)
 		{
 			// Prime the state so that the first "after login" check will occur
@@ -264,6 +297,10 @@ inline FName GetUniqueTag(UUpdateManager* UpdateManager)
 void UUpdateManager::StartPatchCheck()
 {
 	ensure(ChecksEnabled());
+	if (TSharedPtr<IUpdateManagerStatsCollector> StrongStatsCollector = StatsCollector.Pin())
+	{
+		StrongStatsCollector->OnCheckStep_PatchCheckStarted();
+	}
 
 	SetUpdateState(EUpdateState::CheckingForPatch);
 	UGameInstance* GameInstance = GetGameInstance();
@@ -343,6 +380,11 @@ void UUpdateManager::PatchCheckComplete(EPatchCheckResult PatchResult)
 
 void UUpdateManager::InstallBundlePatchCheckComplete(EInstallBundleManagerPatchCheckResult PatchResult)
 {
+	if (TSharedPtr<IUpdateManagerStatsCollector> StrongStatsCollector = StatsCollector.Pin())
+	{
+		StrongStatsCollector->OnCheckStep_PatchCheckComplete(PatchResult);
+	}
+
 	TSharedPtr<IInstallBundleManager> InstallBundleMan = IInstallBundleManager::GetPlatformInstallBundleManager();
 	if (InstallBundleMan && !InstallBundleMan->IsNullInterface())
 	{
@@ -373,6 +415,11 @@ void UUpdateManager::InstallBundlePatchCheckComplete(EInstallBundleManagerPatchC
 
 void UUpdateManager::StartHotfixCheck()
 {
+	if (TSharedPtr<IUpdateManagerStatsCollector> StrongStatsCollector = StatsCollector.Pin())
+	{
+		StrongStatsCollector->OnCheckStep_HotfixCheckStarted();
+	}
+
 	if (bCheckHotfixAvailabilityOnly)
 	{
 		// Just check for the presence of a hotfix
@@ -423,8 +470,13 @@ void UUpdateManager::OnHotfixCheckComplete(EHotfixResult Result)
 
 	LastHotfixResult = Result;
 
-	auto CompletionDelegate = [this]()
+	auto CompletionDelegate = [this, Result]()
 	{
+		if (TSharedPtr<IUpdateManagerStatsCollector> StrongStatsCollector = StatsCollector.Pin())
+		{
+			StrongStatsCollector->OnCheckStep_HotfixCheckComplete(Result);
+		}
+
 		// Always preload data
 		StartInitialPreload();
 	};
@@ -458,6 +510,11 @@ void UUpdateManager::HotfixAvailabilityCheckComplete(EHotfixResult Result)
 	auto CompletionDelegate = [this, Result]()
 	{
 		UE_LOG(LogHotfixManager, Display, TEXT("External HotfixAvailabilityCheckComplete %d"), (int32)Result);
+		if (TSharedPtr<IUpdateManagerStatsCollector> StrongStatsCollector = StatsCollector.Pin())
+		{
+			StrongStatsCollector->OnCheckStep_HotfixCheckComplete(Result);
+		}
+
 		switch (Result)
 		{
 			case EHotfixResult::Success:
@@ -482,6 +539,11 @@ void UUpdateManager::HotfixAvailabilityCheckComplete(EHotfixResult Result)
 
 void UUpdateManager::StartInitialPreload()
 {
+	if (TSharedPtr<IUpdateManagerStatsCollector> StrongStatsCollector = StatsCollector.Pin())
+	{
+		StrongStatsCollector->OnCheckStep_InitialPreloadStarted();
+	}
+
 	SetUpdateState(EUpdateState::WaitingOnInitialLoad);
 
 	// Start ticking
@@ -496,6 +558,11 @@ void UUpdateManager::StartInitialPreload()
 
 void UUpdateManager::InitialPreloadComplete()
 {
+	if (TSharedPtr<IUpdateManagerStatsCollector> StrongStatsCollector = StatsCollector.Pin())
+	{
+		StrongStatsCollector->OnCheckStep_InitialPreloadComplete();
+	}
+
 	SetUpdateState(EUpdateState::InitialLoadComplete);
 
 	if (LastPatchCheckResult == EInstallBundleManagerPatchCheckResult::PatchCheckFailure)
diff --git a/Engine/Plugins/Online/OnlineFramework/Source/Hotfix/Public/UpdateManager.h b/Engine/Plugins/Online/OnlineFramework/Source/Hotfix/Public/UpdateManager.h
--- a/Engine/Plugins/Online/OnlineFramework/Source/Hotfix/Public/UpdateManager.h
+++ b/Engine/Plugins/Online/OnlineFramework/Source/Hotfix/Public/UpdateManager.h
@@ -130,6 +130,38 @@ public:
 	TSet<FString> AdditionalTags;
 };
 
+/**
+ * Collector for aggregating stats related to the update check process.
+ *
+ * Implement collector and pass into RegisterStatsCollector to receive stats.
+ */ 
+class IUpdateManagerStatsCollector
+{
+public:
+	virtual ~IUpdateManagerStatsCollector() = default;
+
+	struct FCheckStartedParams
+	{
+		FString ContextName;
+		bool bEnabled = true;
+		bool bCheckAvailabilityOnly = false;
+		bool bCachedResponse = false;
+		double CacheExpiration = 0.0;
+	};
+
+	// Overall process/
+	virtual void OnCheckStarted(FCheckStartedParams&& Params) = 0;
+	virtual void OnCheckComplete(EUpdateCompletionStatus Result) = 0;
+
+	// Individual stages.
+	virtual void OnCheckStep_PatchCheckStarted() = 0;
+	virtual void OnCheckStep_PatchCheckComplete(EInstallBundleManagerPatchCheckResult PatchResult) = 0;
+	virtual void OnCheckStep_HotfixCheckStarted() = 0;
+	virtual void OnCheckStep_HotfixCheckComplete(EHotfixResult Result) = 0;
+	virtual void OnCheckStep_InitialPreloadStarted() = 0;
+	virtual void OnCheckStep_InitialPreloadComplete() = 0;
+};
+
 /**
  * Update manager
  *
@@ -215,6 +247,11 @@ public:
 	*/
 	bool IsActivelyUpdating() const { return IsUpdating() && (CurrentUpdateState != EUpdateState::UpdatePending); }
 
+	/**
+	 * Register a collector for acquiring stats related to the update check process.
+	 */
+	UE_API void RegisterStatsCollector(const TSharedPtr<IUpdateManagerStatsCollector>& InStatsCollector);
+
 	/**
 	 * @return the current state of the update check process
 	 */
@@ -409,6 +446,8 @@ protected:
 	FTSTicker::FDelegateHandle TickerHandle;
 	FTSTicker::FDelegateHandle StartCheckInternalTimerHandle;
 
+	TWeakPtr<IUpdateManagerStatsCollector> StatsCollector;
+
 private:
 
 	/**
