Remove invitation when join party request completed, instead of removing it immediately.

diff --git a/Engine/Plugins/Online/OnlineFramework/Source/Party/Private/User/SocialUser.cpp b/Engine/Plugins/Online/OnlineFramework/Source/Party/Private/User/SocialUser.cpp
--- a/Engine/Plugins/Online/OnlineFramework/Source/Party/Private/User/SocialUser.cpp
+++ b/Engine/Plugins/Online/OnlineFramework/Source/Party/Private/User/SocialUser.cpp
@@ -113,6 +113,12 @@ TAutoConsoleVariable<FString> USocialUser::FDebugInitializer::CVarSocialUserDefe
 TArray<TWeakObjectPtr<USocialUser>> USocialUser::FDebugInitializer::PausedSocialUsers;
 #endif // !UE_BUILD_SHIPPING
 
+TAutoConsoleVariable<bool> CVarSocialUserRemoveInvitationOnJoinPartyCompleteEnabled(
+	TEXT("party.SocialUserRemoveInvitationOnJoinPartyCompleteEnabled"),
+	true,
+	TEXT("When enabled, will remove invitation on join party complete, instead of remove immediately")
+);
+
 //////////////////////////////////////////////////////////////////////////
 // FSocialQuery_UserInfo
 //////////////////////////////////////////////////////////////////////////
@@ -1484,10 +1490,20 @@ void USocialUser::JoinParty(const FOnlinePartyTypeId& PartyTypeId, const FName&
 {
 	IOnlinePartyJoinInfoConstPtr SentInvite = GetSentPartyInvite(PartyTypeId);
 
-	GetOwningToolkit().GetSocialManager().JoinParty(*this, PartyTypeId, USocialManager::FOnJoinPartyAttemptComplete(), JoinMethod);
+	GetOwningToolkit().GetSocialManager().JoinParty(*this, PartyTypeId, USocialManager::FOnJoinPartyAttemptComplete::CreateUObject(this, &USocialUser::OnJoinPartyComplete, PartyTypeId, SentInvite), JoinMethod);
 
 	// Regardless of the outcome, note that the invite was accepted (deletes it from the OSS party system)
-	if (SentInvite)
+	if (SentInvite && !CVarSocialUserRemoveInvitationOnJoinPartyCompleteEnabled.GetValueOnAnyThread())
+	{
+		IOnlinePartyPtr PartyInterface = Online::GetPartyInterfaceChecked(GetWorld());
+		PartyInterface->ClearInvitations(*GetOwningToolkit().GetLocalUserNetId(ESocialSubsystem::Primary), *GetUserId(ESocialSubsystem::Primary), nullptr);
+		OnPartyInviteAcceptedInternal(PartyTypeId, *SentInvite);
+	}
+}
+
+void USocialUser::OnJoinPartyComplete(const FJoinPartyResult& JoinPartyResult, FOnlinePartyTypeId PartyTypeId, IOnlinePartyJoinInfoConstPtr SentInvite) const
+{
+	if (SentInvite && CVarSocialUserRemoveInvitationOnJoinPartyCompleteEnabled.GetValueOnAnyThread() && JoinPartyResult.GetResult() != EJoinPartyCompletionResult::NoResponse)
 	{
 		IOnlinePartyPtr PartyInterface = Online::GetPartyInterfaceChecked(GetWorld());
 		PartyInterface->ClearInvitations(*GetOwningToolkit().GetLocalUserNetId(ESocialSubsystem::Primary), *GetUserId(ESocialSubsystem::Primary), nullptr);
diff --git a/Engine/Plugins/Online/OnlineFramework/Source/Party/Public/User/SocialUser.h b/Engine/Plugins/Online/OnlineFramework/Source/Party/Public/User/SocialUser.h
--- a/Engine/Plugins/Online/OnlineFramework/Source/Party/Public/User/SocialUser.h
+++ b/Engine/Plugins/Online/OnlineFramework/Source/Party/Public/User/SocialUser.h
@@ -267,6 +267,8 @@ private:
 	bool bIsInitialized = false;
 	UE_API void FinishInitialization();
 
+	UE_API void OnJoinPartyComplete(const FJoinPartyResult& JoinPartyResult, FOnlinePartyTypeId PartyTypeId, IOnlinePartyJoinInfoConstPtr SentInvite) const;
+
 #if !UE_BUILD_SHIPPING
 	// Debug info for initializing social user
 	class FDebugInitializer;
