Allow joining Away player's party if the Away player sent an invitation

#FailedValidations Preflight

diff --git a/Engine/Plugins/Online/OnlineFramework/Source/Party/Private/Party/SocialParty.cpp b/Engine/Plugins/Online/OnlineFramework/Source/Party/Private/Party/SocialParty.cpp
--- a/Engine/Plugins/Online/OnlineFramework/Source/Party/Private/Party/SocialParty.cpp
+++ b/Engine/Plugins/Online/OnlineFramework/Source/Party/Private/Party/SocialParty.cpp
@@ -172,13 +172,29 @@ FPartyJoinApproval USocialParty::EvaluateJoinRequest(const TArray<IOnlinePartyUs
 {
 	FPartyJoinApproval JoinApproval;
 
+	auto CheckForSentInvitation = [this, &Players]()
+	{
+		for (const IOnlinePartyUserPendingJoinRequestInfoConstRef& JoiningPlayer : Players)
+		{
+			for (USocialToolkit* Toolkit : UE::OnlineFramework::GetLocalPartyMemberToolkits(*this))
+			{
+				if (USocialUser* JoiningSocialUser = Toolkit->FindUser(JoiningPlayer->GetUserId());
+					JoiningSocialUser && JoiningSocialUser->HasBeenInvitedToParty(IOnlinePartySystem::GetPrimaryPartyTypeId()))
+				{
+					return true;
+				}
+			}
+		}
+		return false;
+	};
 	if ((GetNumPartyMembers() + Players.Num()) > GetPartyMaxSize())
 	{
 		JoinApproval.SetDenialReason(EPartyJoinDenialReason::PartyFull);
 	}
 	else if (USocialUser* SocialUser = GetOwningLocalMember().GetSocialUser(GetOwningLocalUserId());
 		SocialUser &&
-		(USocialManager::ShouldReadPresenceFromOnlineServices() ? SocialUser->GetOnlineStatusV2() == UE::Online::EUserPresenceStatus::Away : SocialUser->GetOnlineStatus() == EOnlinePresenceState::Away))
+		(USocialManager::ShouldReadPresenceFromOnlineServices() ? SocialUser->GetOnlineStatusV2() == UE::Online::EUserPresenceStatus::Away : SocialUser->GetOnlineStatus() == EOnlinePresenceState::Away) &&
+		(!GetSocialManager().CanAcceptInvitationsFromAwayUsers() || !CheckForSentInvitation()))
 	{
 		JoinApproval.SetDenialReason(EPartyJoinDenialReason::TargetUserAway);
 	}
diff --git a/Engine/Plugins/Online/OnlineFramework/Source/Party/Private/SocialManager.cpp b/Engine/Plugins/Online/OnlineFramework/Source/Party/Private/SocialManager.cpp
--- a/Engine/Plugins/Online/OnlineFramework/Source/Party/Private/SocialManager.cpp
+++ b/Engine/Plugins/Online/OnlineFramework/Source/Party/Private/SocialManager.cpp
@@ -491,7 +491,8 @@ FJoinPartyResult USocialManager::ValidateJoinTarget(const USocialUser& UserToJoi
 	{
 		return FPartyJoinDenialReason(EPartyJoinDenialReason::NotLoggedIn);
 	}
-	else if (ShouldReadPresenceFromOnlineServices() ? UserToJoin.GetOnlineStatusV2() == UE::Online::EUserPresenceStatus::Away : UserToJoin.GetOnlineStatus() == EOnlinePresenceState::Away)
+	else if ((ShouldReadPresenceFromOnlineServices() ? UserToJoin.GetOnlineStatusV2() == UE::Online::EUserPresenceStatus::Away : UserToJoin.GetOnlineStatus() == EOnlinePresenceState::Away) && // Away indicates they do not want to be joined
+		(!CanAcceptInvitationsFromAwayUsers() || !UserToJoin.HasSentPartyInvite(PartyTypeId))) // They can still send you invitations while away, should we ignore the Away status if we've received an invite?
 	{
 		return FPartyJoinDenialReason(EPartyJoinDenialReason::TargetUserAway);
 	}
@@ -572,6 +573,11 @@ FJoinPartyResult USocialManager::ValidateJoinTarget(const USocialUser& UserToJoi
 	return FJoinPartyResult();
 }
 
+bool USocialManager::CanAcceptInvitationsFromAwayUsers() const
+{
+	return bAcceptInvitationsFromAwayUsers;
+}
+
 void USocialManager::JoinParty(const USocialUser& UserToJoin, const FOnlinePartyTypeId& PartyTypeId, const FOnJoinPartyAttemptComplete& OnJoinPartyComplete, const FName& JoinMethod)
 {
 	UE_LOG(LogParty, Verbose, TEXT("Attempting to join user [%s]'s party of type [%d] by [%s]"), *UserToJoin.ToDebugString(), PartyTypeId.GetValue(), *JoinMethod.ToString());
diff --git a/Engine/Plugins/Online/OnlineFramework/Source/Party/Public/SocialManager.h b/Engine/Plugins/Online/OnlineFramework/Source/Party/Public/SocialManager.h
--- a/Engine/Plugins/Online/OnlineFramework/Source/Party/Public/SocialManager.h
+++ b/Engine/Plugins/Online/OnlineFramework/Source/Party/Public/SocialManager.h
@@ -128,6 +128,8 @@ public:
 	UE_API virtual FJoinPartyResult ValidateJoinTarget(const USocialUser& UserToJoin, const FOnlinePartyTypeId& PartyTypeId) const;
 
 	UE_API virtual FJoinPartyResult ValidateJoinTarget(const USocialUser& UserToJoin, const FOnlinePartyTypeId& PartyTypeId, bool bCheckPlatformSession) const;
+	/** @return true if we can join an away user that has sent us an invitation */
+	UE_API bool CanAcceptInvitationsFromAwayUsers() const;
 
 protected:
 	DECLARE_DELEGATE_OneParam(FOnJoinPartyAttemptComplete, const FJoinPartyResult&);
@@ -274,6 +276,10 @@ private:
 	UPROPERTY(Config)
 	FName PresenceFrameworkAdapterInstance;
 
+	/** Can we accept an invitation (ie, join party) from an Away user? */
+	UPROPERTY(config)
+	bool bAcceptInvitationsFromAwayUsers = true;
+
 	bool bIsConnectedToPartyService = false;
 	
 	/**
diff --git a/Engine/Plugins/Online/OnlineFramework/Source/Party/Public/SocialToolkit.h b/Engine/Plugins/Online/OnlineFramework/Source/Party/Public/SocialToolkit.h
--- a/Engine/Plugins/Online/OnlineFramework/Source/Party/Public/SocialToolkit.h
+++ b/Engine/Plugins/Online/OnlineFramework/Source/Party/Public/SocialToolkit.h
@@ -80,8 +80,8 @@ public:
 
 	UE_API const FOnlineUserPresence* GetPresenceInfo(ESocialSubsystem SubsystemType) const;
 	UE_API const UE::Online::TOnlineResult<UE::Online::FGetCachedPresence> GetPresenceInfoV2(ESocialSubsystem SubsystemType) const;
-	UE_API void SetLocalUserOnlineState(EOnlinePresenceState::Type OnlineState);
-	UE_API void SetLocalUserOnlineStateV2(UE::Online::EUserPresenceStatus OnlineState);
+	UE_API virtual void SetLocalUserOnlineState(EOnlinePresenceState::Type OnlineState);
+	UE_API virtual void SetLocalUserOnlineStateV2(UE::Online::EUserPresenceStatus OnlineState);
 
 	UE_DEPRECATED(5.6, "AddLocalUserOnlineProperties is deprecated, please use IOnlinePresence::SetPresence instead.")
 	UE_API void AddLocalUserOnlineProperties(FPresenceProperties OnlineProperties);
diff --git a/Engine/Plugins/Online/OnlineFramework/Source/Party/Public/User/SocialUser.h b/Engine/Plugins/Online/OnlineFramework/Source/Party/Public/User/SocialUser.h
--- a/Engine/Plugins/Online/OnlineFramework/Source/Party/Public/User/SocialUser.h
+++ b/Engine/Plugins/Online/OnlineFramework/Source/Party/Public/User/SocialUser.h
@@ -144,7 +144,7 @@ public:
 	UE_API virtual void JoinParty(const FOnlinePartyTypeId& PartyTypeId, const FName& JoinMethod) const;
 	UE_API virtual void RejectPartyInvite(const FOnlinePartyTypeId& PartyTypeId);
 
-	UE_API bool HasBeenInvitedToParty(const FOnlinePartyTypeId& PartyTypeId) const;
+	UE_API virtual bool HasBeenInvitedToParty(const FOnlinePartyTypeId& PartyTypeId) const;
 	UE_API bool CanInviteToParty(const FOnlinePartyTypeId& PartyTypeId) const;
 	UE_API bool InviteToParty(const FOnlinePartyTypeId& PartyTypeId, const ESocialPartyInviteMethod InviteMethod = ESocialPartyInviteMethod::Other, const FString& MetaData = FString()) const;
 
