Fix for wrong id object being checked

diff --git a/Engine/Plugins/Online/OnlineFramework/Source/Party/Private/User/SocialUser.cpp b/Engine/Plugins/Online/OnlineFramework/Source/Party/Private/User/SocialUser.cpp
--- a/Engine/Plugins/Online/OnlineFramework/Source/Party/Private/User/SocialUser.cpp
+++ b/Engine/Plugins/Online/OnlineFramework/Source/Party/Private/User/SocialUser.cpp
@@ -1335,7 +1335,7 @@ bool USocialUser::IsOnline() const
 				}
 			}
 
-			if(LocalUserAccountId.IsValid() && UserId.IsValid())
+			if(LocalUserAccountId.IsValid() && TargetAccountId.IsValid())
 			{
 				if (IOnlineServicesPtr OnlineServices = CommonConfig.GetServices(PresenceFrameworkInstance))
 				{
@@ -1396,7 +1396,7 @@ bool USocialUser::IsPlayingThisGame() const
 				}
 			}
 
-			if(LocalUserAccountId.IsValid() && UserId.IsValid())
+			if(LocalUserAccountId.IsValid() && TargetAccountId.IsValid())
 			{
 				if (IOnlineServicesPtr OnlineServices = CommonConfig.GetServices(PresenceFrameworkInstance))
 				{
