Make FSocialInteraction_AcceptPartyInvite::CanExecute check if joining the party would be successful.  It only checked if an invitation had been received, and this resulted in many join failures if the target was away, playing a different game, etc.

diff --git a/Engine/Plugins/Online/OnlineFramework/Source/Party/Private/Interactions/PartyInteractions.cpp b/Engine/Plugins/Online/OnlineFramework/Source/Party/Private/Interactions/PartyInteractions.cpp
--- a/Engine/Plugins/Online/OnlineFramework/Source/Party/Private/Interactions/PartyInteractions.cpp
+++ b/Engine/Plugins/Online/OnlineFramework/Source/Party/Private/Interactions/PartyInteractions.cpp
@@ -159,7 +159,7 @@ FString FSocialInteraction_AcceptPartyInvite::GetSlashCommandToken()
 
 bool FSocialInteraction_AcceptPartyInvite::CanExecute(const USocialUser& User)
 {
-	return User.HasSentPartyInvite(IOnlinePartySystem::GetPrimaryPartyTypeId());
+	return User.HasSentPartyInvite(IOnlinePartySystem::GetPrimaryPartyTypeId()) && User.CheckPartyJoinability(IOnlinePartySystem::GetPrimaryPartyTypeId()).WasSuccessful();
 }
 
 void FSocialInteraction_AcceptPartyInvite::ExecuteInteraction(USocialUser& User)
