Added new CVar and logic to disable writing presence to MCP/XMPP

diff --git a/Engine/Plugins/Online/OnlineFramework/Source/Party/Private/PartyModule.cpp b/Engine/Plugins/Online/OnlineFramework/Source/Party/Private/PartyModule.cpp
--- a/Engine/Plugins/Online/OnlineFramework/Source/Party/Private/PartyModule.cpp
+++ b/Engine/Plugins/Online/OnlineFramework/Source/Party/Private/PartyModule.cpp
@@ -20,6 +20,14 @@ void FPartyModule::ShutdownModule()
 {
 }
 
+/**
+* CVar to disable writing information to Online Subsystem Presence
+*/
+static TAutoConsoleVariable<int32> CVar_OnlineFramework_Party_ShouldWritePresenceToOnlineSubsystem(
+	TEXT("onlineframework.party.shouldwritepresencetoonlinesubsystem"),
+	1,
+	TEXT("Whether relevant systems should write Presence information to OnlineSubsystem"));
+
 /**
 * CVar to enable writing information to Online Services Presence
 */
