Include the join info in failed party joins

diff --git a/Engine/Plugins/Online/OnlineFramework/Source/Party/Private/SocialManager.cpp b/Engine/Plugins/Online/OnlineFramework/Source/Party/Private/SocialManager.cpp
--- a/Engine/Plugins/Online/OnlineFramework/Source/Party/Private/SocialManager.cpp
+++ b/Engine/Plugins/Online/OnlineFramework/Source/Party/Private/SocialManager.cpp
@@ -24,11 +24,21 @@
 #include "Misc/Base64.h"
 #endif
 
+#if !UE_BUILD_SHIPPING
 static TAutoConsoleVariable<bool> CVarForceDisconnectedToPartyService(
 	TEXT("SocialUI.ForceDisconnectedToPartyService"),
 	false,
 	TEXT("Overrides the bIsConnectedToPartyService state"));
+#endif
 
+namespace UE::OnlineFramework::CVars
+{
+	bool bAssignPartyJoinInfoOnPartyJoinFailure = true;
+	static FAutoConsoleVariableRef CVarAssignPartyJoinInfoOnPartyJoinFailure(
+		TEXT("Party.AssignPartyJoinInfoOnPartyJoinFailure"), bAssignPartyJoinInfoOnPartyJoinFailure,
+		TEXT("Temporary cvar (remove after proven stable), whether to set the party join info on the join attempt when ValidateJoinTarget fails"),
+		ECVF_Default);
+}
 //////////////////////////////////////////////////////////////////////////
 // FJoinPartyAttempt
 //////////////////////////////////////////////////////////////////////////
@@ -602,6 +612,10 @@ void USocialManager::JoinParty(const USocialUser& UserToJoin, const FOnlineParty
 	}
 	else
 	{
+		if (UE::OnlineFramework::CVars::bAssignPartyJoinInfoOnPartyJoinFailure)
+		{
+			NewAttempt.JoinInfo = UserToJoin.GetPartyJoinInfo(PartyTypeId);
+		}
 		// We don't do the standard FinishJoinAttempt here because this entry isn't actually in our map of join attempts yet
 		// It's possible that this attempt failed immediately because a join is already in progress, in which case we don't want to replace the legitimate attempt with the same ID
 		OnJoinPartyAttemptCompleteInternal(NewAttempt, ValidationResult);
@@ -1197,7 +1211,11 @@ void USocialManager::HandleLeavePartyForJoinComplete(ELeavePartyCompletionResult
 
 bool USocialManager::IsConnectedToPartyService() const
 {
-	return bIsConnectedToPartyService && !CVarForceDisconnectedToPartyService.GetValueOnAnyThread();
+	return bIsConnectedToPartyService 
+#if !UE_BUILD_SHIPPING
+		&& !CVarForceDisconnectedToPartyService.GetValueOnAnyThread()
+#endif
+	;
 }
 
 void USocialManager::HandlePartyDisconnected(USocialParty* DisconnectingParty)
