Online: move FName/FString definitions out of headers

diff --git a/Engine/Plugins/Online/OnlineFramework/Source/Party/Private/Party/PartyTypes.cpp b/Engine/Plugins/Online/OnlineFramework/Source/Party/Private/Party/PartyTypes.cpp
--- a/Engine/Plugins/Online/OnlineFramework/Source/Party/Private/Party/PartyTypes.cpp
+++ b/Engine/Plugins/Online/OnlineFramework/Source/Party/Private/Party/PartyTypes.cpp
@@ -138,3 +138,150 @@ void FOnlinePartyRepDataBase::LogPropertyChanged(const TCHAR* OwningStructTypeNa
 		OwningParty ? *OwningParty->ToDebugString() : TEXT("unknown"));
 }
 
+const TCHAR* ToString(EPartyJoinDenialReason Type)
+{
+	switch (Type)
+	{
+	case EPartyJoinDenialReason::NoReason:
+		return TEXT("NoReason");
+	case EPartyJoinDenialReason::JoinAttemptAborted:
+		return TEXT("JoinAttemptAborted");
+	case EPartyJoinDenialReason::Busy:
+		return TEXT("Busy");
+	case EPartyJoinDenialReason::OssUnavailable:
+		return TEXT("OssUnavailable");
+	case EPartyJoinDenialReason::PartyFull:
+		return TEXT("PartyFull");
+	case EPartyJoinDenialReason::GameFull:
+		return TEXT("GameFull");
+	case EPartyJoinDenialReason::NotPartyLeader:
+		return TEXT("NotPartyLeader");
+	case EPartyJoinDenialReason::PartyPrivate:
+		return TEXT("PartyPrivate");
+	case EPartyJoinDenialReason::JoinerCrossplayRestricted:
+		return TEXT("JoinerCrossplayRestricted");
+	case EPartyJoinDenialReason::MemberCrossplayRestricted:
+		return TEXT("MemberCrossplayRestricted");
+	case EPartyJoinDenialReason::GameModeRestricted:
+		return TEXT("GameModeRestricted");
+	case EPartyJoinDenialReason::Banned:
+		return TEXT("Banned");
+	case EPartyJoinDenialReason::NotLoggedIn:
+		return TEXT("NotLoggedIn");
+	case EPartyJoinDenialReason::CheckingForRejoin:
+		return TEXT("CheckingForRejoin");
+	case EPartyJoinDenialReason::TargetUserMissingPresence:
+		return TEXT("TargetUserMissingPresence");
+	case EPartyJoinDenialReason::TargetUserUnjoinable:
+		return TEXT("TargetUserUnjoinable");
+	case EPartyJoinDenialReason::TargetUserAway:
+		return TEXT("TargetUserAway");
+	case EPartyJoinDenialReason::AlreadyLeaderInPlatformSession:
+		return TEXT("AlreadyLeaderInPlatformSession");
+	case EPartyJoinDenialReason::TargetUserPlayingDifferentGame:
+		return TEXT("TargetUserPlayingDifferentGame");
+	case EPartyJoinDenialReason::TargetUserMissingPlatformSession:
+		return TEXT("TargetUserMissingPlatformSession");
+	case EPartyJoinDenialReason::PlatformSessionMissingJoinInfo:
+		return TEXT("PlatformSessionMissingJoinInfo");
+	case EPartyJoinDenialReason::FailedToStartFindConsoleSession:
+		return TEXT("FailedToStartFindConsoleSession");
+	case EPartyJoinDenialReason::MissingPartyClassForTypeId:
+		return TEXT("MissingPartyClassForTypeId");
+	case EPartyJoinDenialReason::TargetUserBlocked:
+		return TEXT("TargetUserBlocked");
+	case EPartyJoinDenialReason::InvalidJoinInfo:
+		return TEXT("InvalidJoinInfo");
+	case EPartyJoinDenialReason::NotFriends:
+		return TEXT("NotFriends");
+	default:
+		return TEXT("CustomReason");
+	}
+}
+
+const TCHAR* LexToString(ESocialPartyInviteFailureReason Type)
+{
+	switch (Type)
+	{
+	case ESocialPartyInviteFailureReason::Success: return TEXT("Success");
+	case ESocialPartyInviteFailureReason::NotOnline: return TEXT("NotOnline");
+	case ESocialPartyInviteFailureReason::NotAcceptingMembers: return TEXT("NotAcceptingMembers");
+	case ESocialPartyInviteFailureReason::NotFriends: return TEXT("NotFriends");
+	case ESocialPartyInviteFailureReason::AlreadyInParty: return TEXT("AlreadyInParty");
+	case ESocialPartyInviteFailureReason::OssValidationFailed: return TEXT("OssValidationFailed");
+	case ESocialPartyInviteFailureReason::PlatformInviteFailed: return TEXT("PlatformInviteFailed");
+	case ESocialPartyInviteFailureReason::PartyInviteFailed: return TEXT("PartyInviteFailed");
+	case ESocialPartyInviteFailureReason::InviteRateLimitExceeded: return TEXT("InviteRateLimitExceeded");
+	default:
+		checkNoEntry();
+		return TEXT("Unknown");
+	}
+}
+
+const TCHAR* ToString(EPartyType Type)
+{
+	switch (Type)
+	{
+	case EPartyType::Public:
+	{
+		return TEXT("Public");
+	}
+	case EPartyType::FriendsOnly:
+	{
+		return TEXT("FriendsOnly");
+	}
+	case EPartyType::Private:
+	{
+		return TEXT("Private");
+	}
+	default:
+	{
+		return TEXT("Unknown");
+	}
+	}
+}
+
+const TCHAR* ToString(EApprovalAction Type)
+{
+	switch (Type)
+	{
+	case EApprovalAction::Approve:
+	{
+		return TEXT("Approve");
+	}
+	case EApprovalAction::Enqueue:
+	{
+		return TEXT("Enqueue");
+	}
+	case EApprovalAction::EnqueueAndStartBeacon:
+	{
+		return TEXT("EnqueueAndStartBeacon");
+	}
+	case EApprovalAction::Deny:
+	{
+		return TEXT("Deny");
+	}
+	default:
+	{
+		return TEXT("Unknown");
+	}
+	}
+}
+
+namespace PartyJoinMethod
+{
+	// User has created the party
+	const FName Creation(TEXT("Creation"));
+	// User has joined the party via invitation
+	const FName Invitation(TEXT("Invitation"));
+	// user has joined after requesting access
+	const FName RequestToJoin(TEXT("RequestToJoin"));
+	// User has joined the party via presence
+	const FName Presence(TEXT("Presence"));
+	// User has joined the party using a platform option
+	const FName PlatformSession(TEXT("PlatformSession"));
+	// User has joined the party via command line
+	const FName CommandLineJoin(TEXT("CommandLineJoin"));
+	// User has joined via unknown/undocumented process
+	const FName Unspecified(TEXT("Unspecified"));
+}
diff --git a/Engine/Plugins/Online/OnlineFramework/Source/Party/Private/SocialToolkit.cpp b/Engine/Plugins/Online/OnlineFramework/Source/Party/Private/SocialToolkit.cpp
--- a/Engine/Plugins/Online/OnlineFramework/Source/Party/Private/SocialToolkit.cpp
+++ b/Engine/Plugins/Online/OnlineFramework/Source/Party/Private/SocialToolkit.cpp
@@ -18,6 +18,18 @@
 
 #include UE_INLINE_GENERATED_CPP_BY_NAME(SocialToolkit)
 
+namespace FriendInviteFailureReason 
+{
+	const FString InviteFailReason_NotFound = TEXT("NotFound");
+	const FString InviteFailReason_AddingSelfFail = TEXT("AddingSelfFail");
+	const FString InviteFailReason_AddingBlockedFail = TEXT("AddingBlockedFail");
+	const FString InviteFailReason_AlreadyFriends = TEXT("AlreadyFriends");
+}
+namespace FriendAcceptFailureReason
+{
+	const FString AcceptFailReason_NotPendingInbound = TEXT("NotPendingInbound");
+}
+
 #if WITH_EDITOR
 DECLARE_MULTICAST_DELEGATE_TwoParams(FOnStartRandomizeUserPresence, uint8 /*NumRandomUser*/, float /*TickerTimer*/);
 static FOnStartRandomizeUserPresence Debug_OnStartRandomizeUserPresenceEvent;
diff --git a/Engine/Plugins/Online/OnlineFramework/Source/Party/Public/Party/PartyTypes.h b/Engine/Plugins/Online/OnlineFramework/Source/Party/Public/Party/PartyTypes.h
--- a/Engine/Plugins/Online/OnlineFramework/Source/Party/Public/Party/PartyTypes.h
+++ b/Engine/Plugins/Online/OnlineFramework/Source/Party/Public/Party/PartyTypes.h
@@ -145,66 +145,7 @@ enum class EPartyJoinDenialReason : uint8
 	MAX
 };
 
-inline const TCHAR* ToString(EPartyJoinDenialReason Type)
-{
-	switch (Type)
-	{
-	case EPartyJoinDenialReason::NoReason:
-		return TEXT("NoReason");
-	case EPartyJoinDenialReason::JoinAttemptAborted:
-		return TEXT("JoinAttemptAborted");
-	case EPartyJoinDenialReason::Busy:
-		return TEXT("Busy");
-	case EPartyJoinDenialReason::OssUnavailable:
-		return TEXT("OssUnavailable");
-	case EPartyJoinDenialReason::PartyFull:
-		return TEXT("PartyFull");
-	case EPartyJoinDenialReason::GameFull:
-		return TEXT("GameFull");
-	case EPartyJoinDenialReason::NotPartyLeader:
-		return TEXT("NotPartyLeader");
-	case EPartyJoinDenialReason::PartyPrivate:
-		return TEXT("PartyPrivate");
-	case EPartyJoinDenialReason::JoinerCrossplayRestricted:
-		return TEXT("JoinerCrossplayRestricted");
-	case EPartyJoinDenialReason::MemberCrossplayRestricted:
-		return TEXT("MemberCrossplayRestricted");
-	case EPartyJoinDenialReason::GameModeRestricted:
-		return TEXT("GameModeRestricted");
-	case EPartyJoinDenialReason::Banned:
-		return TEXT("Banned");
-	case EPartyJoinDenialReason::NotLoggedIn:
-		return TEXT("NotLoggedIn");
-	case EPartyJoinDenialReason::CheckingForRejoin:
-		return TEXT("CheckingForRejoin");
-	case EPartyJoinDenialReason::TargetUserMissingPresence:
-		return TEXT("TargetUserMissingPresence");
-	case EPartyJoinDenialReason::TargetUserUnjoinable:
-		return TEXT("TargetUserUnjoinable");
-	case EPartyJoinDenialReason::TargetUserAway:
-		return TEXT("TargetUserAway");
-	case EPartyJoinDenialReason::AlreadyLeaderInPlatformSession:
-		return TEXT("AlreadyLeaderInPlatformSession");
-	case EPartyJoinDenialReason::TargetUserPlayingDifferentGame:
-		return TEXT("TargetUserPlayingDifferentGame");
-	case EPartyJoinDenialReason::TargetUserMissingPlatformSession:
-		return TEXT("TargetUserMissingPlatformSession");
-	case EPartyJoinDenialReason::PlatformSessionMissingJoinInfo:
-		return TEXT("PlatformSessionMissingJoinInfo");
-	case EPartyJoinDenialReason::FailedToStartFindConsoleSession:
-		return TEXT("FailedToStartFindConsoleSession");
-	case EPartyJoinDenialReason::MissingPartyClassForTypeId:
-		return TEXT("MissingPartyClassForTypeId");
-	case EPartyJoinDenialReason::TargetUserBlocked:
-		return TEXT("TargetUserBlocked");
-	case EPartyJoinDenialReason::InvalidJoinInfo:
-		return TEXT("InvalidJoinInfo");
-	case EPartyJoinDenialReason::NotFriends:
-		return TEXT("NotFriends");
-	default:
-		return TEXT("CustomReason");
-	}
-}
+PARTY_API const TCHAR* ToString(EPartyJoinDenialReason Type);
 
 UENUM()
 enum class EApprovalAction : uint8
@@ -254,24 +195,7 @@ enum class ESocialPartyInviteFailureReason : uint8
 };
 
 
-inline const TCHAR* LexToString(ESocialPartyInviteFailureReason Type)
-{
-	switch (Type)
-	{
-	case ESocialPartyInviteFailureReason::Success: return TEXT("Success");
-	case ESocialPartyInviteFailureReason::NotOnline: return TEXT("NotOnline");
-	case ESocialPartyInviteFailureReason::NotAcceptingMembers: return TEXT("NotAcceptingMembers");
-	case ESocialPartyInviteFailureReason::NotFriends: return TEXT("NotFriends");
-	case ESocialPartyInviteFailureReason::AlreadyInParty: return TEXT("AlreadyInParty");
-	case ESocialPartyInviteFailureReason::OssValidationFailed: return TEXT("OssValidationFailed");
-	case ESocialPartyInviteFailureReason::PlatformInviteFailed: return TEXT("PlatformInviteFailed");
-	case ESocialPartyInviteFailureReason::PartyInviteFailed: return TEXT("PartyInviteFailed");
-	case ESocialPartyInviteFailureReason::InviteRateLimitExceeded: return TEXT("InviteRateLimitExceeded");
-	default:
-		checkNoEntry();
-		return TEXT("Unknown");
-	}
-}
+PARTY_API const TCHAR* LexToString(ESocialPartyInviteFailureReason Type);
 
 // Gives a smidge more meaning to the intended use for the string. These should just be UniqueId's (and are), but not reliably allocated as shared ptrs, so they cannot be replicated via FUniqueNetIdRepl.
 using FSessionId = FString;
@@ -642,74 +566,26 @@ EXPOSE_REVISED_REP_DATA_PROPERTY_SETTER_ACCESS(Owner, PropertyType, StructProper
 #define EXPOSE_REVISED_USTRUCT_REP_DATA_PROPERTY(Owner, PropertyType, StructProperty, ChildProperty, DeprecatedPropertyName, DeprecatedVersion)	\
 EXPOSE_REVISED_REP_DATA_PROPERTY_SETTER_ACCESS(Owner, PropertyType, StructProperty##ChildProperty, StructProperty.ChildProperty, public, DeprecatedPropertyName, DeprecatedVersion);
 
-inline const TCHAR* ToString(EPartyType Type)
-{
-	switch (Type)
-	{
-	case EPartyType::Public:
-	{
-		return TEXT("Public");
-	}
-	case EPartyType::FriendsOnly:
-	{
-		return TEXT("FriendsOnly");
-	}
-	case EPartyType::Private:
-	{
-		return TEXT("Private");
-	}
-	default:
-	{
-		return TEXT("Unknown");
-	}
-	}
-}
-
-inline const TCHAR* ToString(EApprovalAction Type)
-{
-	switch (Type)
-	{
-	case EApprovalAction::Approve:
-	{
-		return TEXT("Approve");
-	}
-	case EApprovalAction::Enqueue:
-	{
-		return TEXT("Enqueue");
-	}
-	case EApprovalAction::EnqueueAndStartBeacon:
-	{
-		return TEXT("EnqueueAndStartBeacon");
-	}
-	case EApprovalAction::Deny:
-	{
-		return TEXT("Deny");
-	}
-	default:
-	{
-		return TEXT("Unknown");
-	}
-	}
-}
+PARTY_API const TCHAR* ToString(EPartyType Type);
+PARTY_API const TCHAR* ToString(EApprovalAction Type);
 
 /** Defines the means used to join a party*/
 namespace PartyJoinMethod
 {
 	// User has created the party
-	const FName Creation = FName(TEXT("Creation"));
+	extern PARTY_API const FName Creation;
 	// User has joined the party via invitation
-	const FName Invitation = FName(TEXT("Invitation"));
+	extern PARTY_API const FName Invitation;
 	// user has joined after requesting access
-	const FName RequestToJoin = FName(TEXT("RequestToJoin"));
+	extern PARTY_API const FName RequestToJoin;
 	// User has joined the party via presence
-	const FName Presence = FName(TEXT("Presence"));
+	extern PARTY_API const FName Presence;
 	// User has joined the party using a platform option
-	const FName PlatformSession = FName(TEXT("PlatformSession"));
+	extern PARTY_API const FName PlatformSession;
 	// User has joined the party via command line
-	const FName CommandLineJoin = FName(TEXT("CommandLineJoin"));
+	extern PARTY_API const FName CommandLineJoin;
 	// User has joined via unknown/undocumented process
-	const FName Unspecified = FName(TEXT("Unspecified"));
-	
+	extern PARTY_API const FName Unspecified;
 }
 
 #undef UE_API
diff --git a/Engine/Plugins/Online/OnlineFramework/Source/Party/Public/SocialToolkit.h b/Engine/Plugins/Online/OnlineFramework/Source/Party/Public/SocialToolkit.h
--- a/Engine/Plugins/Online/OnlineFramework/Source/Party/Public/SocialToolkit.h
+++ b/Engine/Plugins/Online/OnlineFramework/Source/Party/Public/SocialToolkit.h
@@ -32,14 +32,14 @@ enum class EMemberExitedReason : uint8;
 
 namespace FriendInviteFailureReason 
 {
-	const FString InviteFailReason_NotFound = TEXT("NotFound");
-	const FString InviteFailReason_AddingSelfFail = TEXT("AddingSelfFail");
-	const FString InviteFailReason_AddingBlockedFail = TEXT("AddingBlockedFail");
-	const FString InviteFailReason_AlreadyFriends = TEXT("AlreadyFriends");
+	extern PARTY_API const FString InviteFailReason_NotFound;
+	extern PARTY_API const FString InviteFailReason_AddingSelfFail;
+	extern PARTY_API const FString InviteFailReason_AddingBlockedFail;
+	extern PARTY_API const FString InviteFailReason_AlreadyFriends;
 }
 namespace FriendAcceptFailureReason
 {
-	const FString AcceptFailReason_NotPendingInbound = TEXT("NotPendingInbound");
+	extern PARTY_API const FString AcceptFailReason_NotPendingInbound;
 }
 
 DECLARE_DELEGATE_OneParam(FUserDependentAction, USocialUser&);
diff --git a/Engine/Plugins/Online/OnlineSubsystem/Source/Private/OnlineSubsystemTypes.cpp b/Engine/Plugins/Online/OnlineSubsystem/Source/Private/OnlineSubsystemTypes.cpp
--- a/Engine/Plugins/Online/OnlineSubsystem/Source/Private/OnlineSubsystemTypes.cpp
+++ b/Engine/Plugins/Online/OnlineSubsystem/Source/Private/OnlineSubsystemTypes.cpp
@@ -555,3 +555,5 @@ void ParseOnlineSubsystemConfigPairs(TArrayView<const FString> InEntries, TArray
 		OutPairs.Emplace(MoveTemp(KeyString), MoveTemp(ValueString));
 	}
 }
+
+FName NAME_Unset = TEXT("UNSET");
diff --git a/Engine/Plugins/Online/OnlineSubsystem/Source/Public/OnlineSubsystemTypes.h b/Engine/Plugins/Online/OnlineSubsystem/Source/Public/OnlineSubsystemTypes.h
--- a/Engine/Plugins/Online/OnlineSubsystem/Source/Public/OnlineSubsystemTypes.h
+++ b/Engine/Plugins/Online/OnlineSubsystem/Source/Public/OnlineSubsystemTypes.h
@@ -430,7 +430,7 @@ public:
 };
 
 // placeholder "type" until we can make FUniqueNetIdString sufficiently abstract
-static FName NAME_Unset = TEXT("UNSET");
+extern ONLINESUBSYSTEM_API FName NAME_Unset;
 
 using FUniqueNetIdStringRef = TSharedRef<const class FUniqueNetIdString>;
 using FUniqueNetIdStringPtr = TSharedPtr<const class FUniqueNetIdString>;
diff --git a/Engine/Source/Runtime/Online/HTTP/Private/Interfaces/IHttpRequest.cpp b/Engine/Source/Runtime/Online/HTTP/Private/Interfaces/IHttpRequest.cpp
--- a/Engine/Source/Runtime/Online/HTTP/Private/Interfaces/IHttpRequest.cpp
+++ b/Engine/Source/Runtime/Online/HTTP/Private/Interfaces/IHttpRequest.cpp
@@ -2,6 +2,15 @@
 
 #include "Interfaces/IHttpRequest.h"
 
+namespace HttpRequestOptions
+{
+	const FName HttpVersion("HttpVersion");
+	const FName RequestMode("RequestMode");
+#if UE_HTTP_SUPPORT_UNIX_SOCKET
+	const FName UnixSocketPath("UnixSocketPath");
+#endif //UE_HTTP_SUPPORT_UNIX_SOCKET
+}
+
 const TCHAR* LexToString(EHttpRequestMode Mode)
 {
 	switch (Mode)
diff --git a/Engine/Source/Runtime/Online/HTTP/Public/Interfaces/IHttpRequest.h b/Engine/Source/Runtime/Online/HTTP/Public/Interfaces/IHttpRequest.h
--- a/Engine/Source/Runtime/Online/HTTP/Public/Interfaces/IHttpRequest.h
+++ b/Engine/Source/Runtime/Online/HTTP/Public/Interfaces/IHttpRequest.h
@@ -143,10 +143,10 @@ private:
  */
 namespace HttpRequestOptions
 {
-	static const FName HttpVersion("HttpVersion");
-	static const FName RequestMode("RequestMode");
+	extern HTTP_API const FName HttpVersion;
+	extern HTTP_API const FName RequestMode;
 #if UE_HTTP_SUPPORT_UNIX_SOCKET
-	static const FName UnixSocketPath("UnixSocketPath");
+	extern HTTP_API const FName UnixSocketPath;
 #endif //UE_HTTP_SUPPORT_UNIX_SOCKET
 }
 
diff --git a/Engine/Source/Runtime/Online/HTTP/Public/Interfaces/IHttpResponse.h b/Engine/Source/Runtime/Online/HTTP/Public/Interfaces/IHttpResponse.h
--- a/Engine/Source/Runtime/Online/HTTP/Public/Interfaces/IHttpResponse.h
+++ b/Engine/Source/Runtime/Online/HTTP/Public/Interfaces/IHttpResponse.h
@@ -9,7 +9,7 @@
 
 namespace EHttpResponseCodes
 {
-	static const FString ErrorNamespace = TEXT("errors.com.epicgames.httpresponse");
+	static constexpr TCHAR ErrorNamespace[] = TEXT("errors.com.epicgames.httpresponse");
 	/**
 	 * Response codes that can come back from an Http request
 	 */
@@ -42,7 +42,7 @@ namespace EHttpResponseCodes
 
 	inline FString GetResponseCodeAsErrorCode(int32 StatusCode)
 	{
-		return FString::Printf(TEXT("%s.%s"), *ErrorNamespace, *ResponseCodeToString((Type)StatusCode));
+		return FString::Printf(TEXT("%s.%s"), ErrorNamespace, *ResponseCodeToString((Type)StatusCode));
 	}
 
 	/**
