Fixes for Online Services Presence functionality when not used together with Online Services UserInfo functionality
Added FastTick to PartialUpdatePresence

diff --git a/Engine/Plugins/Online/OnlineFramework/Source/Party/Private/User/SocialUser.cpp b/Engine/Plugins/Online/OnlineFramework/Source/Party/Private/User/SocialUser.cpp
--- a/Engine/Plugins/Online/OnlineFramework/Source/Party/Private/User/SocialUser.cpp
+++ b/Engine/Plugins/Online/OnlineFramework/Source/Party/Private/User/SocialUser.cpp
@@ -1060,7 +1060,7 @@ FUserPlatform USocialUser::GetCurrentPlatform() const
 		bool bIsPrimariPresenceUserPlayingThisGame = false;
 		FString PrimaryPresencePlatformStr;
 
-		if (ShouldReadUserInfoFromOnlineServices())
+		if (USocialManager::ShouldReadPresenceFromOnlineServices() || ShouldReadUserInfoFromOnlineServices())
 		{
 			using namespace UE::Online;
 
@@ -1218,7 +1218,7 @@ void USocialUser::GetRichPresenceText(FText& OutRichPresence) const
 	else if (IsFriend() || IsLocalUser())
 	{
 		bool bIsPrimaryPresenceStatusStrValid = false;
-		if (ShouldReadUserInfoFromOnlineServices())
+		if (USocialManager::ShouldReadPresenceFromOnlineServices() || ShouldReadUserInfoFromOnlineServices())
 		{
 			using namespace UE::Online;
 
@@ -1248,7 +1248,7 @@ void USocialUser::GetRichPresenceText(FText& OutRichPresence) const
 			}
 			else
 			{
-				if (ShouldReadUserInfoFromOnlineServices())
+				if (USocialManager::ShouldReadPresenceFromOnlineServices() || ShouldReadUserInfoFromOnlineServices())
 				{
 					OutRichPresence = ToLocText(GetOnlineStatusV2());
 				}
@@ -1308,7 +1308,7 @@ bool USocialUser::IsOnline() const
 	// If any presence says we're online, count it (note also that only friends have presence info, so non-friends will always count as offline)
 	for (const TPair<ESocialSubsystem, FSubsystemUserInfo>& SubsystemInfoPair : SubsystemInfoByType)
 	{
-		if (ShouldReadUserInfoFromOnlineServices())
+		if (USocialManager::ShouldReadPresenceFromOnlineServices() || ShouldReadUserInfoFromOnlineServices())
 		{
 			using namespace UE::Online;
 			using namespace UE::OnlineFramework;
@@ -1369,7 +1369,7 @@ bool USocialUser::IsPlayingThisGame() const
 {
 	for (const TPair<ESocialSubsystem, FSubsystemUserInfo>& SubsystemInfoPair : SubsystemInfoByType)
 	{
-		if (ShouldReadUserInfoFromOnlineServices())
+		if (USocialManager::ShouldReadPresenceFromOnlineServices() || ShouldReadUserInfoFromOnlineServices())
 		{
 			using namespace UE::Online;
 			using namespace UE::OnlineFramework;
diff --git a/Engine/Plugins/Online/OnlineServicesEOS/Source/Private/Online/PresenceEOS.cpp b/Engine/Plugins/Online/OnlineServicesEOS/Source/Private/Online/PresenceEOS.cpp
--- a/Engine/Plugins/Online/OnlineServicesEOS/Source/Private/Online/PresenceEOS.cpp
+++ b/Engine/Plugins/Online/OnlineServicesEOS/Source/Private/Online/PresenceEOS.cpp
@@ -682,7 +682,9 @@ TOnlineAsyncOpHandle<FPartialUpdatePresence> FPresenceEOS::PartialUpdatePresence
 							Promise.EmplaceValue();
 							return;
 						}
-					}					
+					}
+
+					GetServices<FOnlineServicesEpicCommon>().AddEOSSDKFastTick(InAsyncOp);
 
 					EOS_Presence_SetPresenceOptions SetPresenceOptions = { };
 					SetPresenceOptions.ApiVersion = 1;
@@ -701,6 +703,8 @@ TOnlineAsyncOpHandle<FPartialUpdatePresence> FPresenceEOS::PartialUpdatePresence
 			})
 			.Then([this](TOnlineAsyncOp<FPartialUpdatePresence>& InAsyncOp, const EOS_Presence_SetPresenceCallbackInfo* Data) mutable
 			{
+				GetServices<FOnlineServicesEpicCommon>().RemoveEOSSDKFastTick(InAsyncOp);
+
 				UE_LOG(LogOnlineServices, Warning, TEXT("SetPresenceResult: [%s]"), *LexToString(Data->ResultCode));
 
 				if (Data->ResultCode == EOS_EResult::EOS_Success)
