Allow clients to override subregion, similar to region.

diff --git a/Engine/Plugins/Online/OnlineFramework/Source/Qos/Private/QosRegionManager.cpp b/Engine/Plugins/Online/OnlineFramework/Source/Qos/Private/QosRegionManager.cpp
--- a/Engine/Plugins/Online/OnlineFramework/Source/Qos/Private/QosRegionManager.cpp
+++ b/Engine/Plugins/Online/OnlineFramework/Source/Qos/Private/QosRegionManager.cpp
@@ -140,6 +140,17 @@ FString FRegionQosInstance::GetBestSubregion() const
 
 void FRegionQosInstance::GetSubregionPreferences(TArray<FString>& OutSubregions) const
 {
+#if !UE_BUILD_SHIPPING
+	FString OverrideSubregion;
+	if (FParse::Value(FCommandLine::Get(), TEXT("McpSubregion="), OverrideSubregion))
+	{
+		// Subregion specified on command line
+		OutSubregions.Add(OverrideSubregion);
+		UE_LOG(LogQos, Verbose, TEXT("[FRegionQosInstance::GetSubregionPreferences] overriding subregion to %s"), *OverrideSubregion);
+		return;
+	}
+#endif
+
 	// Presorted for best/recommended result first
 	for (const FDatacenterQosInstance& Option : DatacenterOptions)
 	{
