PC GDK builds can opt-in to using the OnlineFramework XBL party sessions and invite permissions

#rnx

diff --git a/Engine/Config/ConfigRedirects.ini b/Engine/Config/ConfigRedirects.ini
--- a/Engine/Config/ConfigRedirects.ini
+++ b/Engine/Config/ConfigRedirects.ini
@@ -23,6 +23,9 @@
 ; Fix for TextureLODSettings defaults not allowing DeviceProfile parent settings from working (with TextureLODGroups)
 /Script/Engine.TextureLODSettings=GlobalDefaults DeviceProfile
 
+; Fix for OnlineServices social settings migration for PC GDK
+/Script/Party.SocialSettings_XboxCommonGDK=/Script/Party.SocialSettingsXbl
+
 ;; This will look in sections named [Staging] and remap WhitelistConfigFiles to AllowedConfigFiles, etc
 [Staging]
 WhitelistConfigFiles=AllowedConfigFiles
diff --git a/Engine/Plugins/Online/OnlineFramework/Config/BaseOnlineFramework.ini b/Engine/Plugins/Online/OnlineFramework/Config/BaseOnlineFramework.ini
new file mode 100644
--- /dev/null
+++ b/Engine/Plugins/Online/OnlineFramework/Config/BaseOnlineFramework.ini
@@ -0,0 +1,2 @@
+[CoreRedirects]
++ClassRedirects=(OldName="SocialSettings_XboxCommonGDK",NewName="/Script/Party.SocialSettingsXbl")
diff --git a/Engine/Plugins/Online/OnlineFramework/OnlineFramework.uplugin b/Engine/Plugins/Online/OnlineFramework/OnlineFramework.uplugin
--- a/Engine/Plugins/Online/OnlineFramework/OnlineFramework.uplugin
+++ b/Engine/Plugins/Online/OnlineFramework/OnlineFramework.uplugin
@@ -69,6 +69,12 @@
 		{
 			"Name": "OnlineServicesOSSAdapter",
 			"Enabled": true
+		},
+		{
+			"Name": "OnlineSubsystemGDK",
+			"PlatformAllowList": [ "Win64" ],
+			"Enabled": false,
+			"Optional": true
 		}
 	]
 }
diff --git a/Engine/Plugins/Online/OnlineFramework/Source/Party/Party.Build.cs b/Engine/Plugins/Online/OnlineFramework/Source/Party/Party.Build.cs
--- a/Engine/Plugins/Online/OnlineFramework/Source/Party/Party.Build.cs
+++ b/Engine/Plugins/Online/OnlineFramework/Source/Party/Party.Build.cs
@@ -1,14 +1,21 @@
 // Copyright Epic Games, Inc. All Rights Reserved.
 
 using UnrealBuildTool;
+using EpicGames.Core;
 
 public class Party : ModuleRules
 {
+	[ConfigFile(ConfigHierarchyType.Game, "OnlineFrameworkParty")]
+	bool bUseXbl = false;
+
 	public Party(ReadOnlyTargetRules Target) : base(Target)
 	{
+		ConfigCache.ReadSettings(DirectoryReference.FromFile(Target.ProjectFile), Target.Platform, this, Target.CustomConfig);
+
 		PublicDefinitions.Add("PARTY_PLATFORM_SESSIONS_PSN=" + (bUsesPSNSessions ? "1" : "0"));
 		PublicDefinitions.Add("PARTY_PLATFORM_SESSIONS_XBL=" + (bUsesXBLSessions ? "1" : "0"));
 		PublicDefinitions.Add("PARTY_PLATFORM_INVITE_PERMISSIONS=" + (bUsesPlatformInvitePermissions? "1" : "0"));
+		PublicDefinitions.Add("PARTY_PLATFORM_XBL_INVITE_PERMISSIONS=" + (bUsesXblInvitePermissions ? "1" : "0"));
 
 		PCHUsage = ModuleRules.PCHUsageMode.UseExplicitOrSharedPCHs;
 
@@ -30,9 +37,15 @@ public class Party : ModuleRules
 				"OnlineServicesOSSAdapter"
 			}
 			);
+
+		if (bUseXbl)
+		{
+			PrivateDependencyModuleNames.Add("OnlineSubsystemGDK");
+		}
 	}
 
 	protected virtual bool bUsesPSNSessions { get { return false; } }
-	protected virtual bool bUsesXBLSessions { get { return false; } }
-	protected virtual bool bUsesPlatformInvitePermissions { get { return false; } }
+	protected virtual bool bUsesXBLSessions { get { return bUseXbl; } }
+	protected virtual bool bUsesPlatformInvitePermissions { get { return bUseXbl; } }
+	protected virtual bool bUsesXblInvitePermissions { get { return bUseXbl; } }
 }
diff --git a/Engine/Plugins/Online/OnlineFramework/Source/Party/Private/Microsoft/SocialSettingsXbl.cpp b/Engine/Plugins/Online/OnlineFramework/Source/Party/Private/Microsoft/SocialSettingsXbl.cpp
new file mode 100644
--- /dev/null
+++ b/Engine/Plugins/Online/OnlineFramework/Source/Party/Private/Microsoft/SocialSettingsXbl.cpp
@@ -0,0 +1,19 @@
+// Copyright Epic Games, Inc. All Rights Reserved.
+
+#include "SocialSettingsXbl.h"
+
+USocialSettingsXbl::USocialSettingsXbl()
+{
+}
+
+bool USocialSettingsXbl::ShouldApplyPlatformInvitePermissionsLive()
+{
+	const USocialSettingsXbl& SettingsCDO = *GetDefault<USocialSettingsXbl>();
+	return SettingsCDO.bApplyPlatformInvitePermissionsLive;
+}
+
+bool USocialSettingsXbl::ShouldApplyPlatformInvitePermissionsCrossplay()
+{
+	const USocialSettingsXbl& SettingsCDO = *GetDefault<USocialSettingsXbl>();
+	return SettingsCDO.bApplyPlatformInvitePermissionsCrossplay;
+}
\ No newline at end of file
diff --git a/Engine/Plugins/Online/OnlineFramework/Source/Party/Private/Microsoft/SocialSettingsXbl.h b/Engine/Plugins/Online/OnlineFramework/Source/Party/Private/Microsoft/SocialSettingsXbl.h
new file mode 100644
--- /dev/null
+++ b/Engine/Plugins/Online/OnlineFramework/Source/Party/Private/Microsoft/SocialSettingsXbl.h
@@ -0,0 +1,35 @@
+// Copyright Epic Games, Inc. All Rights Reserved.
+
+#pragma once
+
+#include "UObject/Object.h"
+
+#include "SocialSettingsXbl.generated.h"
+
+/**
+ * Config-driven settings object for the social framework.
+ * Only the CDO is ever expected to be used, no instance is ever expected to be created.
+ */
+UCLASS(Config = Game)
+class USocialSettingsXbl : public UObject
+{
+	GENERATED_BODY()
+
+public:
+	USocialSettingsXbl();
+
+	/** Should GDK permissions be checked for invites received from other xbl users */
+	static bool ShouldApplyPlatformInvitePermissionsLive();
+	/** Should GDK permissions be checked for invites received from non-xbl users */
+	static bool ShouldApplyPlatformInvitePermissionsCrossplay();
+
+private:
+
+	/** Should GDK permissions be checked for invites received from other xbl users */
+	UPROPERTY(config)
+	bool bApplyPlatformInvitePermissionsLive = true;
+
+	/** Should GDK permissions be checked for invites received from non-xbl users */
+	UPROPERTY(config)
+	bool bApplyPlatformInvitePermissionsCrossplay = false;
+};
\ No newline at end of file
diff --git a/Engine/Plugins/Online/OnlineFramework/Source/Party/Private/Microsoft/SocialToolkitXbl.cpp b/Engine/Plugins/Online/OnlineFramework/Source/Party/Private/Microsoft/SocialToolkitXbl.cpp
new file mode 100644
--- /dev/null
+++ b/Engine/Plugins/Online/OnlineFramework/Source/Party/Private/Microsoft/SocialToolkitXbl.cpp
@@ -0,0 +1,114 @@
+// Copyright Epic Games, Inc. All Rights Reserved.
+
+#include "SocialToolkit.h"
+
+#if PARTY_PLATFORM_XBL_INVITE_PERMISSIONS
+
+#if !PARTY_PLATFORM_INVITE_PERMISSIONS
+	#error expected PARTY_PLATFORM_INVITE_PERMISSIONS when PARTY_PLATFORM_XBL_INVITE_PERMISSIONS is set
+#endif
+
+#include "OnlineSubsystemGDK.h"
+#include "OnlineSubsystemGDKTypes.h"
+#include "OnlineUserInterfaceGDK.h"
+#include "SocialSettingsXbl.h"
+
+void USocialToolkit::CanReceiveInviteFrom(USocialUser& User, const IOnlinePartyJoinInfoConstRef& Invite, TFunction<void(const bool /*bResult*/)>&& CompletionFunc)
+{	
+	bool bQueryStarted = false;
+	if (FOnlineSubsystemGDK* OssGDK = static_cast<FOnlineSubsystemGDK*>(IOnlineSubsystem::Get(GDK_SUBSYSTEM)))
+	{
+		if (FOnlineUserGDKPtr UserInterface = OssGDK->GetUsersGDK())
+		{
+			FUniqueNetIdRepl LocalUserPlatformId = GetLocalUserNetId(ESocialSubsystem::Platform);
+			if (LocalUserPlatformId.IsValid())
+			{
+				// Need to work out if the invite is from the same platform, so we know which permissions API to call, the one for GDK users, or the one for cross platform
+				const FString LocalExternalAccountType = FUserPlatform(IOnlineSubsystem::GetLocalPlatformName()).GetPlatformDescription().ExternalAccountType;
+				const FString InviterExternalAccountType = FUserPlatform(Invite->GetSourcePlatform()).GetPlatformDescription().ExternalAccountType;
+
+				if(LocalExternalAccountType == InviterExternalAccountType
+					&& USocialSettingsXbl::ShouldApplyPlatformInvitePermissionsLive())
+				{
+					FUniqueNetIdRepl SocialUserPlatformId = User.GetUserId(ESocialSubsystem::Platform);
+					if (ensure(SocialUserPlatformId.IsValid() && (SocialUserPlatformId.GetType() == GDK_SUBSYSTEM)))
+					{
+						UE_LOG(LogParty, Log, TEXT("%s Inviter=[%s] InviterPlatformId=[%s] checking XBL invite permissions"), ANSI_TO_TCHAR(__FUNCTION__), *User.ToDebugString(), *SocialUserPlatformId->ToDebugString());
+
+						// GDK user, so check with the "other GDK user" api
+						const FUniqueNetIdRef& SocialUserPlatformNetIdRef = SocialUserPlatformId.GetUniqueNetId().ToSharedRef();
+						UserInterface->QueryUserCommunicationPermissions(
+							*LocalUserPlatformId.GetUniqueNetId(),
+							{ SocialUserPlatformNetIdRef },
+							// No explicit invite permission, apparently have to check for Text+Voice comms, see https://learn.microsoft.com/gaming/gdk/_content/gc/policies/XR/XR015
+							{ XblPermission::CommunicateUsingText, XblPermission::CommunicateUsingVoice },
+							FOnGDKCommunicationPermissionsQueryComplete::CreateWeakLambda(this,
+								[this, SocialUserPlatformNetIdRef, CompletionFunc](const FOnlineError& RequestStatus, const FUniqueNetIdRef& /*RequestingUser*/, const FCommunicationPermissionResultsMap& Results)
+						{
+							bool bCanReceiveInvites = true;
+							if (RequestStatus.WasSuccessful())
+							{
+								if (const bool* bCanCommunicate = Results.Find(SocialUserPlatformNetIdRef))
+								{
+									bCanReceiveInvites = *bCanCommunicate;
+								}
+							}
+							CompletionFunc(bCanReceiveInvites);
+						}));
+						bQueryStarted = true;
+					}
+					else
+					{
+						UE_LOG(LogParty, Warning, TEXT("%s SocialUserPlatformId=[%s] ExternalAccountTypes match, but NetId unexpected"), ANSI_TO_TCHAR(__FUNCTION__), *SocialUserPlatformId.ToDebugString());
+					}
+				}
+				else if (LocalExternalAccountType != InviterExternalAccountType
+					&& USocialSettingsXbl::ShouldApplyPlatformInvitePermissionsCrossplay())
+				{
+					UE_LOG(LogParty, Log, TEXT("%s Inviter=[%s] checking crossplay invite permissions"), ANSI_TO_TCHAR(__FUNCTION__), *User.ToDebugString());
+
+					// Crossplay user, so check with crossplay api
+					const XblAnonymousUserType UserType = User.IsFriend(ESocialSubsystem::Primary) ? XblAnonymousUserType::CrossNetworkFriend : XblAnonymousUserType::CrossNetworkUser;
+					UserInterface->QueryAnonymousUserCommunicationPermissions(
+						*LocalUserPlatformId.GetUniqueNetId(),
+						{ UserType },
+						{ XblPermission::CommunicateUsingText, XblPermission::CommunicateUsingVoice },
+						FOnGDKAnonymousUserCommunicationPermissionsQueryComplete::CreateWeakLambda(this,
+							[this, UserType, CompletionFunc](const FOnlineError& RequestStatus, const FAnonymousUserCommunicationPermissionResultsMap& Results)
+					{
+						bool bCanReceiveInvites = true;
+						if (RequestStatus.WasSuccessful())
+						{
+							if (const bool* bCanCommunicate = Results.Find(UserType))
+							{
+								bCanReceiveInvites = *bCanCommunicate;
+							}
+						}
+						CompletionFunc(bCanReceiveInvites);
+					}));
+					bQueryStarted = true;
+				}
+			}
+			else
+			{
+				UE_LOG(LogParty, Warning, TEXT("%s LocalUserPlatformId=[%s] invalid"), ANSI_TO_TCHAR(__FUNCTION__), *LocalUserPlatformId.ToDebugString());
+			}
+		}
+		else
+		{
+			UE_LOG(LogParty, Warning, TEXT("%s UserInterface=nullptr"), ANSI_TO_TCHAR(__FUNCTION__));
+		}
+	}
+	else
+	{
+		UE_LOG(LogParty, Warning, TEXT("%s OssGDK=nullptr"), ANSI_TO_TCHAR(__FUNCTION__));
+	}
+
+	if (!bQueryStarted)
+	{
+		CompletionFunc(true);
+	}
+}
+
+#endif // PARTY_PLATFORM_XBL_INVITE_PERMISSIONS
+
