Allow PartyData to async load any required assets before PostReplication.

When enabled, TPartyDataReplicator::ProcessReceivedData will serialize the new data into a pending party data instance and ask it to async load any required assets. This load request is then put into a queue which gets processed when the load completes. The queue is to ensure updates are applied in the correct order.

Once the load completes we update the primary RepDataPtr with the incoming party data and call PostReplicate. This ensures the new data can only be accessed once fully valid.

PostReplicate now also takes an optional pointer to any loaded resources to ensure the party data maintains a reference to them, keeping them in memory.

#rnx

diff --git a/Engine/Plugins/Online/OnlineFramework/Source/Party/Private/Party/PartyDataReplicator.cpp b/Engine/Plugins/Online/OnlineFramework/Source/Party/Private/Party/PartyDataReplicator.cpp
--- a/Engine/Plugins/Online/OnlineFramework/Source/Party/Private/Party/PartyDataReplicator.cpp
+++ b/Engine/Plugins/Online/OnlineFramework/Source/Party/Private/Party/PartyDataReplicator.cpp
@@ -5,6 +5,24 @@
 #include "Party/PartyMember.h"
 #include "OnlineSubsystemUtils.h"
 
+namespace PartyDataCvars
+{
+	static bool bAsyncLoadingEnabled = true;
+	FAutoConsoleVariableRef CVarAsyncLoadingEnabled(
+		TEXT("Party.ReplicationAsyncLoadingEnabled"),
+		bAsyncLoadingEnabled,
+		TEXT("When enabled, party data will be given a chance to async load and required assets upon replication.")
+	);
+}
+
+namespace UE::Party
+{
+	bool IsRepDataAsyncLoadingEnabled()
+	{
+		return PartyDataCvars::bAsyncLoadingEnabled;
+	}
+}
+
 void FPartyDataReplicatorHelper::ReplicateDataToMembers(const FOnlinePartyRepDataBase& RepDataInstance, const UScriptStruct& RepDataType, const FOnlinePartyData& ReplicationPayload)
 {
 	if (const USocialParty* OwnerParty = static_cast<const FOnlinePartyRepDataBase*>(&RepDataInstance)->GetOwnerParty())
diff --git a/Engine/Plugins/Online/OnlineFramework/Source/Party/Private/Party/PartyMember.cpp b/Engine/Plugins/Online/OnlineFramework/Source/Party/Private/Party/PartyMember.cpp
--- a/Engine/Plugins/Online/OnlineFramework/Source/Party/Private/Party/PartyMember.cpp
+++ b/Engine/Plugins/Online/OnlineFramework/Source/Party/Private/Party/PartyMember.cpp
@@ -333,7 +333,7 @@ void UPartyMember::NotifyMemberDataReceived(const FOnlinePartyData& MemberData)
 	UE_LOG(LogParty, VeryVerbose, TEXT("Received updated rep data for member [%s]"), *ToDebugString());
 
 	check(MemberDataReplicator.IsValid());
-	MemberDataReplicator.ProcessReceivedData(MemberData/*, bHasReceivedInitialData*/);
+	MemberDataReplicator.ProcessReceivedData(this, MemberData/*, bHasReceivedInitialData*/);
 
 	if (EnumHasAnyFlags(InitializingFlags, EInitializingFlags::InitialMemberData))
 	{
diff --git a/Engine/Plugins/Online/OnlineFramework/Source/Party/Private/Party/SocialParty.cpp b/Engine/Plugins/Online/OnlineFramework/Source/Party/Private/Party/SocialParty.cpp
--- a/Engine/Plugins/Online/OnlineFramework/Source/Party/Private/Party/SocialParty.cpp
+++ b/Engine/Plugins/Online/OnlineFramework/Source/Party/Private/Party/SocialParty.cpp
@@ -863,7 +863,7 @@ void USocialParty::HandlePartyDataReceived(const FUniqueNetId& LocalUserId, cons
 	if (PartyId == GetPartyId())
 	{
 		check(PartyDataReplicator.IsValid());
-		PartyDataReplicator.ProcessReceivedData(PartyData);
+		PartyDataReplicator.ProcessReceivedData(this, PartyData);
 		if (!bHasReceivedRepData)
 		{
 			bHasReceivedRepData = true;
diff --git a/Engine/Plugins/Online/OnlineFramework/Source/Party/Public/Party/PartyDataReplicator.h b/Engine/Plugins/Online/OnlineFramework/Source/Party/Public/Party/PartyDataReplicator.h
--- a/Engine/Plugins/Online/OnlineFramework/Source/Party/Public/Party/PartyDataReplicator.h
+++ b/Engine/Plugins/Online/OnlineFramework/Source/Party/Public/Party/PartyDataReplicator.h
@@ -2,11 +2,20 @@
 
 #pragma once
 
+#include "Engine/AssetManager.h"
+#include "Engine/StreamableManager.h"
+#include "Misc/NotNull.h"
 #include "Party/PartyTypes.h"
 #include "PartyModule.h"
 #include "UObject/GCObject.h"
 #include "Containers/Ticker.h"
 #include "Interfaces/OnlinePartyInterface.h"
+#include "Templates/SharedPointer.h"
+
+namespace UE::Party
+{
+	PARTY_API bool IsRepDataAsyncLoadingEnabled();
+}
 
 /** Util exclusively for use by TPartyDataReplicator to circumvent circular include header issues (we can't include SocialParty.h or PartyMember.h here) */
 class FPartyDataReplicatorHelper
@@ -55,8 +64,9 @@ public:
 		RepDataPtr = &RepDataInstance;
 		RepDataType = ChildRepDataT::StaticStruct();
 
-		RepDataCopy = (ChildRepDataT*)FMemory::Malloc(RepDataType->GetCppStructOps()->GetSize());
+		RepDataCopy = (ChildRepDataT*)FMemory::Malloc(RepDataType->GetCppStructOps()->GetSize(), RepDataType->GetCppStructOps()->GetAlignment());
 		RepDataType->GetCppStructOps()->Construct(RepDataCopy);
+		RepDataType->GetCppStructOps()->Copy(RepDataCopy, RepDataPtr, 1);
 	}
 
 	void Flush()
@@ -72,24 +82,76 @@ public:
 	}
 
 protected:
-	void ProcessReceivedData(const FOnlinePartyData& IncomingPartyData, bool bCompareToPrevious = true)
+	void ProcessReceivedData(TNotNull<OwningObjectT*> Owner, const FOnlinePartyData& IncomingPartyData, bool bCompareToPrevious = true)
 	{
 		// If the rep data can be edited locally, disregard any replication updates (they're the same at best or out of date at worst)
 		if (!static_cast<FOnlinePartyRepDataBase*>(RepDataPtr)->CanEditData())
 		{
-			if (FVariantDataConverter::VariantMapToUStruct(IncomingPartyData.GetKeyValAttrs(), RepDataType, RepDataPtr, 0, CPF_Transient | CPF_RepSkip))
+			if (UE::Party::IsRepDataAsyncLoadingEnabled())
 			{
-				static_cast<FOnlinePartyRepDataBase*>(RepDataPtr)->PostReplication(*RepDataCopy);
+				// Since IncomingPartyData will only contain modified fields, it's important that we apply each update (in order) 
+				// that comes in or the local party data will get out of sync with the server.
+				// This means that if we receive new data while still loading the previous request, rather than canceling the previous 
+				// request, we start a new one to be applied after the previous has completed.
+
+				// Lazily create once to avoid allocs each time data is received.
+				if (!PendingRepData)
+				{
+					// Allocate a new object to write the new data into as we don't want to overwrite the actual data until loading has completed.
+					PendingRepData = (RepDataT*)FMemory::Malloc(RepDataType->GetCppStructOps()->GetSize(), RepDataType->GetCppStructOps()->GetAlignment());
+					RepDataType->GetCppStructOps()->Construct(PendingRepData);
+					// We want this to match the current rep data so it stays in sync and we don't miss any assets in GetAssetsToAsyncLoad.
+					RepDataType->GetCppStructOps()->Copy(PendingRepData, RepDataPtr, 1);
+				}
+
+				// Deserialize the incoming data into the scratch copy so we can extract any assets that need to be loaded.
+				if (FVariantDataConverter::VariantMapToUStruct(IncomingPartyData.GetKeyValAttrs(), RepDataType, PendingRepData, 0, CPF_Transient | CPF_RepSkip))
+				{
+					FStreamableAsyncLoadParams LoadParams;
+					// We always request all assets so the new load handle has everything that needs to stay in memory.
+					static_cast<FOnlinePartyRepDataBase*>(PendingRepData)->GetAssetsToAsyncLoad(LoadParams.TargetsToStream);
+
+					TSharedPtr<FStreamableHandle> LoadHandle = UAssetManager::GetStreamableManager().RequestAsyncLoad(MoveTemp(LoadParams));
+
+					static int32 NextId = 0;
+					const int32 RequestId = ++NextId;
+
+					FQueuedRepData QueuedData;
+					QueuedData.LoadHandle = LoadHandle; // Intentionally not moved into QueuedData as we use the handle after moving QueuedData into the Queue.
+					QueuedData.KeyValueAttrs = IncomingPartyData.GetKeyValAttrs();
+					QueuedData.Id = RequestId;
+					QueuedData.bCompareToPrevious = bCompareToPrevious;
+
+					// We always queue the data even if there is nothing to load to ensure updates are applied in the correct order.
+					RepDataLoadQueue.Push(MoveTemp(QueuedData));
+
+					if (LoadHandle.IsValid() && LoadHandle->IsLoadingInProgress())
+					{
+						UE_LOG(LogParty, Verbose, TEXT("Queueing RepData load request %d for %s. Loading %d assets."), RequestId, *(Owner->GetName()), LoadParams.TargetsToStream.Num());
+
+						// Our lifetime is tied to our owner, so use it to prevent the callback from firing after we've been destroyed.
+						LoadHandle->BindCompleteDelegate(
+							FStreamableDelegate::CreateWeakLambda(Owner, [this]
+							{
+								ProcessLoadQueue();
+							}));
+					}
+					else
+					{
+						UE_LOG(LogParty, Verbose, TEXT("Queueing RepData request %d for %s. All required assets already loaded."), RequestId, *(Owner->GetName()));
 
-				if (bCompareToPrevious)
+						// Try pumping the queue immediately since there's no load to wait on.
+						ProcessLoadQueue();
+					}
+				}
+				else
 				{
-					static_cast<FOnlinePartyRepDataBase*>(RepDataPtr)->CompareAgainst(*RepDataCopy);
+					UE_LOG(LogParty, Error, TEXT("Failed to serialize received party data!"));
 				}
-				ensure(RepDataType->GetCppStructOps()->Copy(RepDataCopy, RepDataPtr, 1));
 			}
 			else
 			{
-				UE_LOG(LogParty, Error, TEXT("Failed to serialize received party data!"));
+				PostReceiveData(IncomingPartyData.GetKeyValAttrs(), bCompareToPrevious);
 			}
 		}
 	}
@@ -101,11 +163,20 @@ protected:
 			static_cast<FOnlinePartyRepDataBase*>(RepDataPtr)->OnDataChanged.Unbind();
 			RepDataPtr = nullptr;
 		}
-		if (RepDataType && RepDataCopy)
+		if (RepDataType)
 		{
-			RepDataType->GetCppStructOps()->Destruct(RepDataCopy);
-			FMemory::Free(RepDataCopy);
-			RepDataCopy = nullptr;
+			if (PendingRepData)
+			{
+				RepDataType->GetCppStructOps()->Destruct(PendingRepData);
+				FMemory::Free(PendingRepData);
+				PendingRepData = nullptr;
+			}
+			if (RepDataCopy)
+			{
+				RepDataType->GetCppStructOps()->Destruct(RepDataCopy);
+				FMemory::Free(RepDataCopy);
+				RepDataCopy = nullptr;
+			}
 			RepDataType = nullptr;
 		}
 		if (UpdateTickerHandle.IsValid())
@@ -113,6 +184,7 @@ protected:
 			FTSTicker::RemoveTicker(UpdateTickerHandle);
 			UpdateTickerHandle.Reset();
 		}
+		RepDataLoadQueue.Empty();
 	}
 
 private:
@@ -142,6 +214,75 @@ private:
 		return false;
 	}
 
+	void PostReceiveData(const FOnlineKeyValuePairs<FString, FVariantData>& KeyValueAttrs, bool bCompareToPrevious, TSharedPtr<FStreamableHandle> ResourceHandle = nullptr)
+	{
+		bool bRepDataUpdated = false;
+		if (UE::Party::IsRepDataAsyncLoadingEnabled())
+		{
+			// Only apply the values to fields whose value matches between RepDataPtr and RepDataCopy.
+			// This prevents stomping locally set fields between when the load was queued and now.
+			bRepDataUpdated = FVariantDataConverter::ApplyVariantMapToUnchangedProperties(KeyValueAttrs, RepDataType, RepDataPtr, RepDataCopy, 0, CPF_Transient | CPF_RepSkip);
+		}
+		else
+		{
+			bRepDataUpdated = FVariantDataConverter::VariantMapToUStruct(KeyValueAttrs, RepDataType, RepDataPtr, 0, CPF_Transient | CPF_RepSkip);
+		}
+
+		if (bRepDataUpdated)
+		{
+			static_cast<FOnlinePartyRepDataBase*>(RepDataPtr)->PostReplication(*RepDataCopy, MoveTemp(ResourceHandle));
+			if (bCompareToPrevious)
+			{
+				static_cast<FOnlinePartyRepDataBase*>(RepDataPtr)->CompareAgainst(*RepDataCopy);
+			}
+			ensure(RepDataType->GetCppStructOps()->Copy(RepDataCopy, RepDataPtr, 1));
+		}
+		else
+		{
+			UE_LOG(LogParty, Error, TEXT("Failed to serialize received party data!"));
+		}
+	}
+
+	void ProcessLoadQueue()
+	{
+		int32 Index = 0;
+
+		// Apply queued updates in order provided they're done loading.
+		while (Index < RepDataLoadQueue.Num())
+		{
+			FQueuedRepData& Request = RepDataLoadQueue[Index];
+
+			// We won't always be waiting on a load as some updates don't contain assets that needed to be loaded first.
+			if (Request.LoadHandle.IsValid() && Request.LoadHandle->IsLoadingInProgress())
+			{
+				UE_LOG(LogParty, Verbose, TEXT("Waiting for RepData request %d to finish loading. %d requests in the queue."), Request.Id, RepDataLoadQueue.Num());
+				break;
+			}
+
+			UE_LOG(LogParty, Verbose, TEXT("Applying loaded RepData request %d."), Request.Id);
+			PostReceiveData(Request.KeyValueAttrs, Request.bCompareToPrevious, MoveTemp(Request.LoadHandle));
+
+			++Index;
+		}
+
+		if (Index > 0)
+		{
+			// Remove processed
+			RepDataLoadQueue.RemoveAt(0, Index, EAllowShrinking::No);
+		}
+	}
+
+	struct FQueuedRepData
+	{
+		TSharedPtr<FStreamableHandle> LoadHandle;
+		FOnlineKeyValuePairs<FString, FVariantData> KeyValueAttrs;
+		int32 Id = 0;
+		bool bCompareToPrevious = false;
+	};
+
+	/** New rep data that is currently async loading any required assets before it can be applied. */
+	TArray<FQueuedRepData> RepDataLoadQueue;
+
 	/** Reflection data for child USTRUCT */
 	TObjectPtr<const UScriptStruct> RepDataType = nullptr;
 
@@ -154,5 +295,8 @@ private:
 	/** Scratch copy of child UStruct for handling replication comparisons */
 	RepDataT* RepDataCopy = nullptr;
 
+	/** Scratch copy of child UStruct for storing incoming rep data to query it prior to replication. */
+	RepDataT* PendingRepData = nullptr;
+
 	FTSTicker::FDelegateHandle UpdateTickerHandle;
 };
diff --git a/Engine/Plugins/Online/OnlineFramework/Source/Party/Public/Party/PartyTypes.h b/Engine/Plugins/Online/OnlineFramework/Source/Party/Public/Party/PartyTypes.h
--- a/Engine/Plugins/Online/OnlineFramework/Source/Party/Public/Party/PartyTypes.h
+++ b/Engine/Plugins/Online/OnlineFramework/Source/Party/Public/Party/PartyTypes.h
@@ -2,8 +2,8 @@
 
 #pragma once
 
-
 #include "GameFramework/OnlineReplStructs.h"
+#include "Templates/SharedPointer.h"
 #include "PartyTypes.generated.h"
 
 #define UE_API PARTY_API
@@ -14,6 +14,8 @@ class UPartyMember;
 class IOnlinePartySystem;
 class FOnlinePartyData;
 
+struct FStreamableHandle;
+
 enum class EJoinPartyCompletionResult : int8;
 
 UENUM()
@@ -371,6 +373,12 @@ struct FOnlinePartyRepDataBase
 	virtual ~FOnlinePartyRepDataBase() {}
 
 protected:
+	/**
+	 * Gets any assets that should be async loaded before PostReplication is called.
+	 * This is only used if Party.ReplicationAsyncLoadingEnabled is enabled.
+	 */
+	virtual void GetAssetsToAsyncLoad(TArray<FSoftObjectPath>& OutAssets) const {}
+
 	/**
 	 * Compare this data against the given old data, triggering delegates as appropriate.
 	 * Intended to be used in concert with EXPOSE_REP_DATA_PROPERTY so that all you need to do is run through each property's CompareX function.
@@ -380,8 +388,12 @@ protected:
 	virtual void CompareAgainst(const FOnlinePartyRepDataBase& OldData) const {}
 
 	/**
-	 * Called directly after an updated member state is received and copied into the local state
+	 * Called directly after an updated member state is received and copied into the local state.
+	 * InResourceHandle will contain any resources returned by GetAssetsToAsyncLoad that were loaded.
+	 * This handle should be retained to keep these resources in memory.
 	 */
+	virtual void PostReplication(const FOnlinePartyRepDataBase& OldData, TSharedPtr<FStreamableHandle> InResourceHandle) {}
+	UE_DEPRECATED(5.8, "Please use version taking InResourceHandle")
 	virtual void PostReplication(const FOnlinePartyRepDataBase& OldData) {}
 	virtual bool CanEditData() const { return false; }
 	virtual const USocialParty* GetOwnerParty() const { return nullptr; }
diff --git a/Engine/Plugins/Online/OnlineSubsystem/Source/Private/OnlineKeyValuePair.cpp b/Engine/Plugins/Online/OnlineSubsystem/Source/Private/OnlineKeyValuePair.cpp
--- a/Engine/Plugins/Online/OnlineSubsystem/Source/Private/OnlineKeyValuePair.cpp
+++ b/Engine/Plugins/Online/OnlineSubsystem/Source/Private/OnlineKeyValuePair.cpp
@@ -1288,6 +1288,11 @@ bool FVariantData::operator!=(const FVariantData& Other) const
 }
 
 bool FVariantDataConverter::VariantMapToUStruct(const FOnlineKeyValuePairs<FString, FVariantData>& VariantMap, const UStruct* StructDefinition, void* OutStruct, int64 CheckFlags, int64 SkipFlags)
+{
+	return ApplyVariantMapToUnchangedProperties(VariantMap, StructDefinition, OutStruct, nullptr, CheckFlags, SkipFlags);
+}
+
+bool FVariantDataConverter::ApplyVariantMapToUnchangedProperties(const FOnlineKeyValuePairs<FString, FVariantData>& VariantMap, const UStruct* StructDefinition, void* OutStruct, const void* CompareStruct, int64 CheckFlags, int64 SkipFlags)
 {
 	for (TFieldIterator<FProperty> PropIt(StructDefinition); PropIt; ++PropIt)
 	{
@@ -1313,6 +1318,17 @@ bool FVariantDataConverter::VariantMapToUStruct(const FOnlineKeyValuePairs<FStri
 		}
 
 		void* Value = Property->ContainerPtrToValuePtr<uint8>(OutStruct);
+
+		if (CompareStruct)
+		{
+			const void* CompareValue = Property->ContainerPtrToValuePtr<uint8>(CompareStruct);
+			if (!Property->Identical(Value, CompareValue))
+			{
+				// Skip applying if the properties don't match.
+				continue;
+			}
+		}
+
 		if (!VariantDataToFProperty(VariantData, Property, Value, CheckFlags, SkipFlags))
 		{
 			UE_LOG_ONLINE(Error, TEXT("VariantMapToUStruct - Unable to parse %s.%s from Variant"), *StructDefinition->GetName(), *PropertyName);
diff --git a/Engine/Plugins/Online/OnlineSubsystem/Source/Public/OnlineKeyValuePair.h b/Engine/Plugins/Online/OnlineSubsystem/Source/Public/OnlineKeyValuePair.h
--- a/Engine/Plugins/Online/OnlineSubsystem/Source/Public/OnlineKeyValuePair.h
+++ b/Engine/Plugins/Online/OnlineSubsystem/Source/Public/OnlineKeyValuePair.h
@@ -590,6 +590,20 @@ public:
 	 */
 	static UE_API bool VariantMapToUStruct(const FOnlineKeyValuePairs<FString, FVariantData>& VariantMap, const UStruct* StructDefinition, void* OutStruct, int64 CheckFlags, int64 SkipFlags);
 
+	/**
+	 * Apply FVariantData elements to properties of a UStruct that are identical with the previous value.
+	 *
+	 * @param VariantMap Input variant data
+	 * @param StructDefinition layout of the UStruct
+	 * @param OutStruct output container for UStruct data
+	 * @param CompareStruct the struct to compare OutStruct to when deciding if a property should be updated.
+	 * @param CheckFlags property must have this flag to be serialized
+	 * @param SkipFlags property cannot have this flag to be serialized
+	 * 
+	 * @return true if it was successful, false otherwise
+	 */
+	static UE_API bool ApplyVariantMapToUnchangedProperties(const FOnlineKeyValuePairs<FString, FVariantData>& VariantMap, const UStruct* StructDefinition, void* OutStruct, const void* CompareStruct, int64 CheckFlags, int64 SkipFlags);
+
 	/**
 	 * Convert an FVariantData to a FProperty
 	 *
