PartyDataReplicator - Explicitly cancel pending loads when resetting, and ensure PendingRepData is updated to match RepData (to account for local changes) before deserializing the new update into it.

Also add some cpu scopes to the heavier functions.

#rnx

diff --git a/Engine/Plugins/Online/OnlineFramework/Source/Party/Public/Party/PartyDataReplicator.h b/Engine/Plugins/Online/OnlineFramework/Source/Party/Public/Party/PartyDataReplicator.h
--- a/Engine/Plugins/Online/OnlineFramework/Source/Party/Public/Party/PartyDataReplicator.h
+++ b/Engine/Plugins/Online/OnlineFramework/Source/Party/Public/Party/PartyDataReplicator.h
@@ -7,6 +7,7 @@
 #include "Misc/NotNull.h"
 #include "Party/PartyTypes.h"
 #include "PartyModule.h"
+#include "ProfilingDebugging/CpuProfilerTrace.h"
 #include "UObject/GCObject.h"
 #include "Containers/Ticker.h"
 #include "Interfaces/OnlinePartyInterface.h"
@@ -84,6 +85,8 @@ public:
 protected:
 	void ProcessReceivedData(TNotNull<OwningObjectT*> Owner, const FOnlinePartyData& IncomingPartyData, bool bCompareToPrevious = true)
 	{
+		TRACE_CPUPROFILER_EVENT_SCOPE_STR("TPartyDataReplicator::ProcessReceivedData");
+
 		// If the rep data can be edited locally, disregard any replication updates (they're the same at best or out of date at worst)
 		if (!static_cast<FOnlinePartyRepDataBase*>(RepDataPtr)->CanEditData())
 		{
@@ -100,10 +103,12 @@ protected:
 					// Allocate a new object to write the new data into as we don't want to overwrite the actual data until loading has completed.
 					PendingRepData = (RepDataT*)FMemory::Malloc(RepDataType->GetCppStructOps()->GetSize(), RepDataType->GetCppStructOps()->GetAlignment());
 					RepDataType->GetCppStructOps()->Construct(PendingRepData);
-					// We want this to match the current rep data so it stays in sync and we don't miss any assets in GetAssetsToAsyncLoad.
-					RepDataType->GetCppStructOps()->Copy(PendingRepData, RepDataPtr, 1);
 				}
 
+				// We want this to match the current rep data so it stays in sync and we don't miss any assets in GetAssetsToAsyncLoad.
+				// Any incoming updates will be applied via IncomingPartyData, but that doesn't account for local changes which is why the copy is needed each time.
+				RepDataType->GetCppStructOps()->Copy(PendingRepData, RepDataPtr, 1);
+
 				// Deserialize the incoming data into the scratch copy so we can extract any assets that need to be loaded.
 				if (FVariantDataConverter::VariantMapToUStruct(IncomingPartyData.GetKeyValAttrs(), RepDataType, PendingRepData, 0, CPF_Transient | CPF_RepSkip))
 				{
@@ -184,6 +189,14 @@ protected:
 			FTSTicker::RemoveTicker(UpdateTickerHandle);
 			UpdateTickerHandle.Reset();
 		}
+
+		for (FQueuedRepData& Request : RepDataLoadQueue)
+		{
+			if (Request.LoadHandle.IsValid())
+			{
+				Request.LoadHandle->CancelHandle();
+			}
+		}
 		RepDataLoadQueue.Empty();
 	}
 
@@ -216,6 +229,8 @@ private:
 
 	void PostReceiveData(const FOnlineKeyValuePairs<FString, FVariantData>& KeyValueAttrs, bool bCompareToPrevious, TSharedPtr<FStreamableHandle> ResourceHandle = nullptr)
 	{
+		TRACE_CPUPROFILER_EVENT_SCOPE_STR("TPartyDataReplicator::PostReceiveData");
+
 		bool bRepDataUpdated = false;
 		if (UE::Party::IsRepDataAsyncLoadingEnabled())
 		{
