Add event when a party leader changes in USocialParty
Add validate options struct to USocialManager::ValidateJoinTarget. The options structure includes the existing option for checking for a platform session, and a new option for whether it's for an active join.  This is to allow more strict rules in different scenarios.  The known scenario now is we want to disable "Join Party" UI options based on some friend's hints, but if we get past that part, we still want to do a best effort join.

diff --git a/Engine/Plugins/Online/OnlineFramework/Source/Party/Private/Party/SocialParty.cpp b/Engine/Plugins/Online/OnlineFramework/Source/Party/Private/Party/SocialParty.cpp
--- a/Engine/Plugins/Online/OnlineFramework/Source/Party/Private/Party/SocialParty.cpp
+++ b/Engine/Plugins/Online/OnlineFramework/Source/Party/Private/Party/SocialParty.cpp
@@ -974,6 +974,8 @@ void USocialParty::HandlePartyMemberPromoted(const FUniqueNetId& LocalUserId, co
 
 		// Now that the leader is gone and a new leader established, make sure the accepting state is correct
 		RefreshPublicJoinability();
+
+		OnPartyLeaderChangedEvent.Broadcast(NewLeader);
 	}
 }
 
diff --git a/Engine/Plugins/Online/OnlineFramework/Source/Party/Private/SocialManager.cpp b/Engine/Plugins/Online/OnlineFramework/Source/Party/Private/SocialManager.cpp
--- a/Engine/Plugins/Online/OnlineFramework/Source/Party/Private/SocialManager.cpp
+++ b/Engine/Plugins/Online/OnlineFramework/Source/Party/Private/SocialManager.cpp
@@ -464,7 +464,7 @@ FJoinPartyResult USocialManager::ValidateJoinAttempt(const FOnlinePartyTypeId& P
 
 FJoinPartyResult USocialManager::ValidateJoinTarget(const USocialUser& UserToJoin, const FOnlinePartyTypeId& PartyTypeId) const
 {
-	return ValidateJoinTarget(UserToJoin, PartyTypeId, /*bCheckPlatformSession*/true);
+	return ValidateJoinTarget(UserToJoin, PartyTypeId, FValidateJoinTargetOptions{});
 }
 
 bool USocialManager::ShouldReadPresenceFromOnlineServices()
@@ -488,6 +488,11 @@ FName USocialManager::GetPresenceFrameworkAdapterInstance()
 }
 
 FJoinPartyResult USocialManager::ValidateJoinTarget(const USocialUser& UserToJoin, const FOnlinePartyTypeId& PartyTypeId, bool bCheckPlatformSession) const
+{
+	return ValidateJoinTarget(UserToJoin, PartyTypeId, FValidateJoinTargetOptions{ .bCheckPlatformSession = bCheckPlatformSession });
+}
+
+FJoinPartyResult USocialManager::ValidateJoinTarget(const USocialUser& UserToJoin, const FOnlinePartyTypeId& PartyTypeId, const FValidateJoinTargetOptions& ValidateOptions) const
 {
 	UE_LOG(LogParty, VeryVerbose, TEXT("Validating user [%s] as join target of party type [%d]"), *UserToJoin.ToDebugString(), PartyTypeId.GetValue());
 
@@ -553,7 +558,7 @@ FJoinPartyResult USocialManager::ValidateJoinTarget(const USocialUser& UserToJoi
 					}
 					else if (!PlatformPresence->SessionId.IsValid() || !PlatformPresence->SessionId->IsValid())
 					{
-						if (bCheckPlatformSession)
+						if (ValidateOptions.bCheckPlatformSession)
 						{
 							return FPartyJoinDenialReason(EPartyJoinDenialReason::TargetUserMissingPlatformSession);
 						}
@@ -594,7 +599,7 @@ void USocialManager::JoinParty(const USocialUser& UserToJoin, const FOnlineParty
 
 	FJoinPartyAttempt NewAttempt(&UserToJoin, PartyTypeId, JoinMethod, OnJoinPartyComplete);
 	NewAttempt.AnalyticsContext = UserToJoin.GetAnalyticsContext();
-	const FJoinPartyResult ValidationResult = ValidateJoinTarget(UserToJoin, PartyTypeId, /*bCheckPlatformSession*/true);
+	const FJoinPartyResult ValidationResult = ValidateJoinTarget(UserToJoin, PartyTypeId, FValidateJoinTargetOptions{});
 	if (ValidationResult.WasSuccessful())
 	{
 		FJoinPartyAttempt& JoinAttempt = JoinAttemptsByTypeId.Add(PartyTypeId, NewAttempt);
diff --git a/Engine/Plugins/Online/OnlineFramework/Source/Party/Private/User/SocialUser.cpp b/Engine/Plugins/Online/OnlineFramework/Source/Party/Private/User/SocialUser.cpp
--- a/Engine/Plugins/Online/OnlineFramework/Source/Party/Private/User/SocialUser.cpp
+++ b/Engine/Plugins/Online/OnlineFramework/Source/Party/Private/User/SocialUser.cpp
@@ -1649,7 +1649,7 @@ bool USocialUser::EndFriendship(ESocialSubsystem SocialSubsystem) const
 
 FJoinPartyResult USocialUser::CheckPartyJoinability(const FOnlinePartyTypeId& PartyTypeId, bool bCheckPlatformSession/*= true*/) const
 {
-	return GetOwningToolkit().GetSocialManager().ValidateJoinTarget(*this, PartyTypeId, bCheckPlatformSession);
+	return GetOwningToolkit().GetSocialManager().ValidateJoinTarget(*this, PartyTypeId, USocialManager::FValidateJoinTargetOptions{ .bCheckPlatformSession = bCheckPlatformSession, .bIsActivelyJoining = false });
 }
 
 bool USocialUser::ShowPlatformProfile()
diff --git a/Engine/Plugins/Online/OnlineFramework/Source/Party/Public/Party/SocialParty.h b/Engine/Plugins/Online/OnlineFramework/Source/Party/Public/Party/SocialParty.h
--- a/Engine/Plugins/Online/OnlineFramework/Source/Party/Public/Party/SocialParty.h
+++ b/Engine/Plugins/Online/OnlineFramework/Source/Party/Public/Party/SocialParty.h
@@ -205,6 +205,13 @@ public:
 	DECLARE_EVENT_TwoParams(USocialParty, FOnPartyMemberLeft, UPartyMember*, const EMemberExitedReason);
 	FOnPartyMemberLeft& OnPartyMemberLeft() const { return OnPartyMemberLeftEvent; }
 
+	/**
+	 * Event triggered when the party leader changes.
+	 * @param NewLeader the new leader. Generally non-null, but there may be race conditions where we get the leader id before we get the leader's data.
+	 */
+	DECLARE_EVENT_OneParam(USocialParty, FOnPartyLeaderChanged, UPartyMember* /*NewLeader*/);
+	FOnPartyLeaderChanged& OnPartyLeaderChanged() const { return OnPartyLeaderChangedEvent; }
+
 	DECLARE_EVENT_OneParam(USocialParty, FOnPartyConfigurationChanged, const FPartyConfiguration&);
 	FOnPartyConfigurationChanged& OnPartyConfigurationChanged() const { return OnPartyConfigurationChangedEvent; }
 
@@ -505,6 +512,7 @@ private:
 	mutable FDisconnectPartyEvent OnPartyDisconnectedEvent;
 	mutable FOnPartyMemberCreated OnPartyMemberCreatedEvent;
 	mutable FOnPartyMemberLeft OnPartyMemberLeftEvent;
+	mutable FOnPartyLeaderChanged OnPartyLeaderChangedEvent;
 	mutable FOnPartyConfigurationChanged OnPartyConfigurationChangedEvent;
 	mutable FOnPartyStateChanged OnPartyStateChangedEvent;
 	mutable FOnPartyMemberConnectionStatusChanged OnPartyMemberConnectionStatusChangedEvent;
diff --git a/Engine/Plugins/Online/OnlineFramework/Source/Party/Public/SocialManager.h b/Engine/Plugins/Online/OnlineFramework/Source/Party/Public/SocialManager.h
--- a/Engine/Plugins/Online/OnlineFramework/Source/Party/Public/SocialManager.h
+++ b/Engine/Plugins/Online/OnlineFramework/Source/Party/Public/SocialManager.h
@@ -123,11 +123,20 @@ public:
 
 	UE_API virtual void NotifyPartyInitialized(USocialParty& Party);
 
-	/** Validates that the target user has valid join info for us to use and that we can join any party of the given type */
 	UE_DEPRECATED(5.6, "ValidateJoinTarget has been deprecated, use the overloaded one instead")
 	UE_API virtual FJoinPartyResult ValidateJoinTarget(const USocialUser& UserToJoin, const FOnlinePartyTypeId& PartyTypeId) const;
-
+	UE_DEPRECATED(5.7, "Use ValidateJoinTarget that takes an options struct")
 	UE_API virtual FJoinPartyResult ValidateJoinTarget(const USocialUser& UserToJoin, const FOnlinePartyTypeId& PartyTypeId, bool bCheckPlatformSession) const;
+	
+	struct FValidateJoinTargetOptions
+	{
+		/** Should we check the platform session for joinability info if the social user is on the same platform as us? */
+		bool bCheckPlatformSession = true;
+		/** Is this part of an active join flow? There may be stricter checks when this is false (eg, we don't want to show Join UI options), but if we get past that, we may want to continue with a best effort join attempt. */
+		bool bIsActivelyJoining = true;
+	};
+	/** Validates that the target user has valid join info for us to use and that we can join any party of the given type */
+	UE_API virtual FJoinPartyResult ValidateJoinTarget(const USocialUser& UserToJoin, const FOnlinePartyTypeId& PartyTypeId, const FValidateJoinTargetOptions& ValidateOptions) const;
 	/** @return true if we can join an away user that has sent us an invitation */
 	UE_API bool CanAcceptInvitationsFromAwayUsers() const;
 
