Add USocialParty::EvaluateJoinRequestAwayCheck so subclasses can extend the Away check logic

diff --git a/Engine/Plugins/Online/OnlineFramework/Source/Party/Private/Party/SocialParty.cpp b/Engine/Plugins/Online/OnlineFramework/Source/Party/Private/Party/SocialParty.cpp
--- a/Engine/Plugins/Online/OnlineFramework/Source/Party/Private/Party/SocialParty.cpp
+++ b/Engine/Plugins/Online/OnlineFramework/Source/Party/Private/Party/SocialParty.cpp
@@ -172,29 +172,11 @@ FPartyJoinApproval USocialParty::EvaluateJoinRequest(const TArray<IOnlinePartyUs
 {
 	FPartyJoinApproval JoinApproval;
 
-	auto CheckForSentInvitation = [this, &Players]()
-	{
-		for (const IOnlinePartyUserPendingJoinRequestInfoConstRef& JoiningPlayer : Players)
-		{
-			for (USocialToolkit* Toolkit : UE::OnlineFramework::GetLocalPartyMemberToolkits(*this))
-			{
-				if (USocialUser* JoiningSocialUser = Toolkit->FindUser(JoiningPlayer->GetUserId());
-					JoiningSocialUser && JoiningSocialUser->HasBeenInvitedToParty(IOnlinePartySystem::GetPrimaryPartyTypeId()))
-				{
-					return true;
-				}
-			}
-		}
-		return false;
-	};
 	if ((GetNumPartyMembers() + Players.Num()) > GetPartyMaxSize())
 	{
 		JoinApproval.SetDenialReason(EPartyJoinDenialReason::PartyFull);
 	}
-	else if (USocialUser* SocialUser = GetOwningLocalMember().GetSocialUser(GetOwningLocalUserId());
-		SocialUser &&
-		(USocialManager::ShouldReadPresenceFromOnlineServices() ? SocialUser->GetOnlineStatusV2() == UE::Online::EUserPresenceStatus::Away : SocialUser->GetOnlineStatus() == EOnlinePresenceState::Away) &&
-		(!GetSocialManager().CanAcceptInvitationsFromAwayUsers() || !CheckForSentInvitation()))
+	else if (!EvaluateJoinRequestAwayCheck(Players, bFromJoinRequest))
 	{
 		JoinApproval.SetDenialReason(EPartyJoinDenialReason::TargetUserAway);
 	}
@@ -223,6 +205,29 @@ FPartyJoinApproval USocialParty::EvaluateJoinRequest(const TArray<IOnlinePartyUs
 	return JoinApproval;
 }
 
+bool USocialParty::EvaluateJoinRequestAwayCheck(const TArray<IOnlinePartyUserPendingJoinRequestInfoConstRef>& Players, bool bFromJoinRequest) const
+{
+	auto CheckForSentInvitation = [this, &Players]()
+	{
+		for (const IOnlinePartyUserPendingJoinRequestInfoConstRef& JoiningPlayer : Players)
+		{
+			for (USocialToolkit* Toolkit : UE::OnlineFramework::GetLocalPartyMemberToolkits(*this))
+			{
+				if (USocialUser* JoiningSocialUser = Toolkit->FindUser(JoiningPlayer->GetUserId());
+					JoiningSocialUser && JoiningSocialUser->HasBeenInvitedToParty(IOnlinePartySystem::GetPrimaryPartyTypeId()))
+				{
+					return true;
+				}
+			}
+		}
+		return false;
+	};
+
+	USocialUser* SocialUser = GetOwningLocalMember().GetSocialUser(GetOwningLocalUserId());
+	const bool bIsAway = SocialUser && (USocialManager::ShouldReadPresenceFromOnlineServices() ? SocialUser->GetOnlineStatusV2() == UE::Online::EUserPresenceStatus::Away : SocialUser->GetOnlineStatus() == EOnlinePresenceState::Away);
+	return (!bIsAway) || (GetSocialManager().CanAcceptInvitationsFromAwayUsers() && CheckForSentInvitation());
+}
+
 bool USocialParty::IsCurrentlyLeaving() const
 {
 	return bIsLeavingParty;
diff --git a/Engine/Plugins/Online/OnlineFramework/Source/Party/Public/Party/SocialParty.h b/Engine/Plugins/Online/OnlineFramework/Source/Party/Public/Party/SocialParty.h
--- a/Engine/Plugins/Online/OnlineFramework/Source/Party/Public/Party/SocialParty.h
+++ b/Engine/Plugins/Online/OnlineFramework/Source/Party/Public/Party/SocialParty.h
@@ -301,6 +301,11 @@ protected:
 
 	/** Determines the joinability of this party for a group of users requesting to join */
 	UE_API virtual FPartyJoinApproval EvaluateJoinRequest(const TArray<IOnlinePartyUserPendingJoinRequestInfoConstRef>& Players, bool bFromJoinRequest) const;
+	/**
+	 * Perform the Away status check for an incoming join request
+	 * @return true if the join request passes the away check
+	 */
+	UE_API virtual bool EvaluateJoinRequestAwayCheck(const TArray<IOnlinePartyUserPendingJoinRequestInfoConstRef>& Players, bool bFromJoinRequest) const;
 
 	/** Determines the reason why, if at all, this party is currently flat-out unjoinable  */
 	UE_API virtual FPartyJoinDenialReason DetermineCurrentJoinability() const;
