Added logic to FSocialQuery_UserInfoV2 to properly trigger completion delegates in error states.
Added additional error handling to UserInfoEpicAccount.

diff --git a/Engine/Plugins/Online/OnlineFramework/Source/Party/Private/User/SocialUser.cpp b/Engine/Plugins/Online/OnlineFramework/Source/Party/Private/User/SocialUser.cpp
--- a/Engine/Plugins/Online/OnlineFramework/Source/Party/Private/User/SocialUser.cpp
+++ b/Engine/Plugins/Online/OnlineFramework/Source/Party/Private/User/SocialUser.cpp
@@ -135,6 +135,9 @@ public:
 		using namespace UE::Online;
 		using namespace UE::OnlineFramework;
 
+		TArray<FAccountId> SocialUserAccountIds;
+		CompletionCallbacksByUserId.GenerateKeyArray(SocialUserAccountIds);
+							
 		if (const USocialToolkit* StrongToolkit = Toolkit.Get())
 		{
 			FCommonConfig CommonConfig(StrongToolkit);
@@ -151,19 +154,17 @@ public:
 						{
 							FAccountId LocalUserAccountId = LocalUserCommonAccountPtr->GetId(USocialManager::GetPresenceFrameworkInstance());
 
-							TArray<FAccountId> SocialUserAccountIds;
-							CompletionCallbacksByUserId.GenerateKeyArray(SocialUserAccountIds);
 							UE_LOG(LogParty, Log, TEXT("FSocialQuery_UserInfoV2 executing for [%d] users on subsystem [%s]"), SocialUserAccountIds.Num(), ToString(SubsystemType));
 
 							UserInfo->QueryUserInfo({ LocalUserAccountId, SocialUserAccountIds })
 								.OnComplete([this, LocalUserAccountId, SocialUserAccountIds, UserInfoWeakPtr = UserInfo.ToWeakPtr()](const TOnlineResult<FQueryUserInfo>& Result)
 								{
+									UE::Online::FOnlineError Error = Errors::Success();
+
 									if (Result.IsOk())
 									{
-										UE_LOG(LogParty, Log, TEXT("FSocialQuery_UserInfoV2 completed query on services [%s]"), ToString(SubsystemType));
-											
-										using namespace UE::Online;
-
+										UE_LOG(LogParty, VeryVerbose, TEXT("FSocialQuery_UserInfoV2 completed query on services [%s]"), ToString(SubsystemType));
+										
 										if (IUserInfoPtr UserInfo = UserInfoWeakPtr.Pin())
 										{
 											// Notify users of the query completion
@@ -185,20 +186,50 @@ public:
 													}
 												}
 											}
-										}
 
-										OnQueryCompleted.ExecuteIfBound(GetQueryId(), AsShared());
+											OnQueryCompleted.ExecuteIfBound(GetQueryId(), AsShared());
+
+											return;
+										}
+										else
+										{
+											Error = Errors::InvalidState();
+										}
 									}
 									else
 									{
-										UE_LOG(LogParty, Log, TEXT("FSocialQuery_UserInfoV2 failed query on services [%s] with error [%s]"), ToString(SubsystemType), *Result.GetErrorValue().GetLogString());
+										Error = Result.GetErrorValue();
+
+										UE_LOG(LogParty, Log, TEXT("FSocialQuery_UserInfoV2 failed query on services [%s] with error [%s]"), ToString(SubsystemType), *ToLogString(Error));										
 									}
+
+									for (const FAccountId& UserAccountId : SocialUserAccountIds)
+									{
+										if (const FOnQueryComplete* CompletionCallback = CompletionCallbacksByUserId.Find(UserAccountId))
+										{
+											CompletionCallback->ExecuteIfBound(SubsystemType, false, MakeError(Error));
+										}
+									}
+
+									OnQueryCompleted.ExecuteIfBound(GetQueryId(), AsShared());
 								});
+
+								return;
 						}
 					}
 				}
 			}
 		}
+
+		for (const FAccountId& UserAccountId : SocialUserAccountIds)
+		{
+			if (const FOnQueryComplete* CompletionCallback = CompletionCallbacksByUserId.Find(UserAccountId))
+			{
+				CompletionCallback->ExecuteIfBound(SubsystemType, false, MakeError(Errors::InvalidState()));
+			}
+		}
+
+		OnQueryCompleted.ExecuteIfBound(GetQueryId(), AsShared());
 	}
 };
 
