Remove the join attempt in FinishJoinPartyAttempt before calling OnJoinPartyAttemptCompleteInternal, to allow for further analysis of the join failure such as calling ValidateJoinTarget
null check PartySessionManager in USocialManager::JoinParty. This is only set on some platforms.

diff --git a/Engine/Plugins/Online/OnlineFramework/Source/Party/Private/SocialManager.cpp b/Engine/Plugins/Online/OnlineFramework/Source/Party/Private/SocialManager.cpp
--- a/Engine/Plugins/Online/OnlineFramework/Source/Party/Private/SocialManager.cpp
+++ b/Engine/Plugins/Online/OnlineFramework/Source/Party/Private/SocialManager.cpp
@@ -38,6 +38,12 @@ namespace UE::OnlineFramework::CVars
 		TEXT("Party.AssignPartyJoinInfoOnPartyJoinFailure"), bAssignPartyJoinInfoOnPartyJoinFailure,
 		TEXT("Temporary cvar (remove after proven stable), whether to set the party join info on the join attempt when ValidateJoinTarget fails"),
 		ECVF_Default);
+
+	bool bRemoveJoinAttemptInFinishFirst = true;
+	static FAutoConsoleVariableRef CVarRemoveJoinAttemptInFinishFirst(
+		TEXT("Party.RemoveJoinAttemptInFinishFirst"), bRemoveJoinAttemptInFinishFirst,
+		TEXT("Temporary cvar (remove after proven stable), whether to remove the join attempt in FinishJoinPartyAttempt before any other calls"),
+		ECVF_Default);
 }
 //////////////////////////////////////////////////////////////////////////
 // FJoinPartyAttempt
@@ -599,7 +605,8 @@ void USocialManager::JoinParty(const USocialUser& UserToJoin, const FOnlineParty
 
 	FJoinPartyAttempt NewAttempt(&UserToJoin, PartyTypeId, JoinMethod, OnJoinPartyComplete);
 	NewAttempt.AnalyticsContext = UserToJoin.GetAnalyticsContext();
-	const FJoinPartyResult ValidationResult = ValidateJoinTarget(UserToJoin, PartyTypeId, FValidateJoinTargetOptions{});
+	FJoinPartyResult ValidationResult = ValidateJoinTarget(UserToJoin, PartyTypeId, FValidateJoinTargetOptions{});
+	bool bFailed = false;
 	if (ValidationResult.WasSuccessful())
 	{
 		FJoinPartyAttempt& JoinAttempt = JoinAttemptsByTypeId.Add(PartyTypeId, NewAttempt);
@@ -609,19 +616,29 @@ void USocialManager::JoinParty(const USocialUser& UserToJoin, const FOnlineParty
 		{
 			QueryPartyJoinabilityInternal(JoinAttempt);
 		}
-		else
+		else if (PartySessionManager)
 		{
 			JoinAttempt.ActionTimeTracker.BeginStep(FJoinPartyAttempt::Step_FindPlatformSession);
 			PartySessionManager->FindSession(UserToJoin, FPartyPlatformSessionManager::FOnFindSessionAttemptComplete::CreateUObject(this, &USocialManager::HandleFindSessionForJoinComplete, PartyTypeId));
 		}
+		else
+		{
+			bFailed = true;
+			ValidationResult = FJoinPartyResult(EJoinPartyCompletionResult::UnknownClientFailure);
+			JoinAttemptsByTypeId.Remove(PartyTypeId);
+		}
 	}
 	else
+	{
+		bFailed = true;
+	}
+	if (bFailed)
 	{
 		if (UE::OnlineFramework::CVars::bAssignPartyJoinInfoOnPartyJoinFailure)
 		{
 			NewAttempt.JoinInfo = UserToJoin.GetPartyJoinInfo(PartyTypeId);
 		}
-		// We don't do the standard FinishJoinAttempt here because this entry isn't actually in our map of join attempts yet
+		// We don't do the standard FinishJoinPartyAttempt here because this entry isn't actually in our map of join attempts yet
 		// It's possible that this attempt failed immediately because a join is already in progress, in which case we don't want to replace the legitimate attempt with the same ID
 		OnJoinPartyAttemptCompleteInternal(NewAttempt, ValidationResult);
 		NewAttempt.OnJoinComplete.ExecuteIfBound(ValidationResult);
@@ -809,14 +826,23 @@ void USocialManager::JoinPartyInternal(FJoinPartyAttempt& JoinAttempt)
 
 void USocialManager::FinishJoinPartyAttempt(FJoinPartyAttempt& JoinAttemptToDestroy, const FJoinPartyResult& JoinResult)
 {
-	OnJoinPartyAttemptCompleteInternal(JoinAttemptToDestroy, JoinResult);
-	JoinAttemptToDestroy.OnJoinComplete.ExecuteIfBound(JoinResult);
-
 	const bool bWasPersistentPartyJoinAttempt = JoinAttemptToDestroy.PartyTypeId == IOnlinePartySystem::GetPrimaryPartyTypeId();
+	if (UE::OnlineFramework::CVars::bRemoveJoinAttemptInFinishFirst)
+	{
+		FJoinPartyAttempt MovedJoinAttempt = MoveTemp(JoinAttemptToDestroy); // JoinAttemptToDestroy is garbage after this. Use MovedJoinAttempt.
+		JoinAttemptsByTypeId.Remove(MovedJoinAttempt.PartyTypeId);
+		OnJoinPartyAttemptCompleteInternal(MovedJoinAttempt, JoinResult);
+		MovedJoinAttempt.OnJoinComplete.ExecuteIfBound(JoinResult);
+	}
+	else
+	{
+		OnJoinPartyAttemptCompleteInternal(JoinAttemptToDestroy, JoinResult);
+		JoinAttemptToDestroy.OnJoinComplete.ExecuteIfBound(JoinResult);
 
-	// JoinAttemptToDestroy is garbage after this. Be careful!
-	JoinAttemptsByTypeId.Remove(JoinAttemptToDestroy.PartyTypeId);
 
+		// JoinAttemptToDestroy is garbage after this. Be careful!
+		JoinAttemptsByTypeId.Remove(JoinAttemptToDestroy.PartyTypeId);
+	}
 	if (bWasPersistentPartyJoinAttempt && !JoinResult.WasSuccessful() && !GetPersistentParty())
 	{
 		// Something goofed when trying to join a new persistent party, so create a replacement immediately
