Add support for AnyWithLocal in ESocialRelationship

This CL adds real support for 'AnyWithLocal' in ESocialRelationship. Similar to 'Any' but without excluding the local user.

#rnx

#tests Tested in PIE.

diff --git a/Engine/Plugins/Online/OnlineFramework/Source/Party/Private/User/SocialUser.cpp b/Engine/Plugins/Online/OnlineFramework/Source/Party/Private/User/SocialUser.cpp
--- a/Engine/Plugins/Online/OnlineFramework/Source/Party/Private/User/SocialUser.cpp
+++ b/Engine/Plugins/Online/OnlineFramework/Source/Party/Private/User/SocialUser.cpp
@@ -440,7 +440,8 @@ TArray<ESocialSubsystem> USocialUser::GetRelationshipSubsystems(ESocialRelations
 			switch (Relationship)
 			{
 			case ESocialRelationship::Any:
-				if (!IsLocalUser())
+			case ESocialRelationship::AnyWithLocal:
+				if (Relationship == ESocialRelationship::AnyWithLocal || !IsLocalUser())
 				{
 					RelationshipSubsystems.Add(SubsystemInfoPair.Key);
 				}
diff --git a/Engine/Plugins/Online/OnlineFramework/Source/Party/Private/User/SocialUserList.cpp b/Engine/Plugins/Online/OnlineFramework/Source/Party/Private/User/SocialUserList.cpp
--- a/Engine/Plugins/Online/OnlineFramework/Source/Party/Private/User/SocialUserList.cpp
+++ b/Engine/Plugins/Online/OnlineFramework/Source/Party/Private/User/SocialUserList.cpp
@@ -20,6 +20,7 @@ FSocialUserList::FSocialUserList(const USocialToolkit& InOwnerToolkit, const FSo
 {
 	if (HasPresenceFilters() && ListConfig.RequiredPresenceFlags != ESocialUserStateFlags::SameParty && ListConfig.ForbiddenPresenceFlags != ESocialUserStateFlags::SameParty &&
 		ListConfig.RelationshipType != ESocialRelationship::Any &&
+		ListConfig.RelationshipType != ESocialRelationship::AnyWithLocal &&
 		ListConfig.RelationshipType != ESocialRelationship::Friend &&
 		ListConfig.RelationshipType != ESocialRelationship::PartyInvite)
 	{
diff --git a/Engine/Plugins/Online/OnlineFramework/Source/Party/Public/SocialTypes.h b/Engine/Plugins/Online/OnlineFramework/Source/Party/Public/SocialTypes.h
--- a/Engine/Plugins/Online/OnlineFramework/Source/Party/Public/SocialTypes.h
+++ b/Engine/Plugins/Online/OnlineFramework/Source/Party/Public/SocialTypes.h
@@ -54,6 +54,7 @@ enum class ESocialRelationship : uint8
 	SuggestedFriend,
 	RecentPlayer,
 	JoinRequest,
+	AnyWithLocal,
 	// Follower, (?)
 };
 
@@ -62,6 +63,7 @@ inline const TCHAR* LexToString(ESocialRelationship Relationship)
 	switch (Relationship)
 	{
 	case ESocialRelationship::Any: return TEXT("Any");
+	case ESocialRelationship::AnyWithLocal: return TEXT("AnyWithLocal");
 	case ESocialRelationship::FriendInviteReceived: return TEXT("FriendInviteReceived");
 	case ESocialRelationship::FriendInviteSent: return TEXT("FriendInviteSent");
 	case ESocialRelationship::PartyInvite: return TEXT("PartyInvite");
