Add support for custom config in hotfixes
MyTarget.CustomConfig = "MyConfig";
-> MyConfig_Engine.ini etc.

diff --git a/Engine/Plugins/Online/OnlineFramework/Source/Hotfix/Private/OnlineHotfixManager.cpp b/Engine/Plugins/Online/OnlineFramework/Source/Hotfix/Private/OnlineHotfixManager.cpp
--- a/Engine/Plugins/Online/OnlineFramework/Source/Hotfix/Private/OnlineHotfixManager.cpp
+++ b/Engine/Plugins/Online/OnlineFramework/Source/Hotfix/Private/OnlineHotfixManager.cpp
@@ -15,6 +15,7 @@
 #include "Misc/CoreDelegates.h"
 #include "Misc/App.h"
 #include "Misc/ConfigUtilities.h"
+#include "Misc/ConfigCacheIni.h"
 
 #include "MoviePlayerProxy.h"
 
@@ -41,6 +42,8 @@ DEFINE_LOG_CATEGORY(LogHotfixManager);
 #define HOTFIX_VERSION_TAG TEXT("Ver-")
 /** The prefix for any hotfix file that expects to indicate branch version information */
 #define HOTFIX_BRANCH_VERSION_TAG TEXT("Branch-")
+/** The prefix used when there no custom config so we don't try to match empty strings */
+#define HOTFIX_NO_CUSTOM_CONFIG TEXT("##NOCONFIG##")
 
 FName NAME_HotfixManager(TEXT("HotfixManager"));
 
@@ -277,6 +280,10 @@ void UOnlineHotfixManager::PostInitProperties()
 	PlatformPrefix += HOTFIX_SEPARATOR;
 	// Server prefix
 	ServerPrefix = DebugPrefix + GetDedicatedServerPrefix();
+	// Custom Config prefix
+	CustomConfigPrefix = FConfigCacheIni::GetCustomConfigString().IsEmpty() ? HOTFIX_NO_CUSTOM_CONFIG : (DebugPrefix + FConfigCacheIni::GetCustomConfigString());
+	CustomConfigPrefix += HOTFIX_SEPARATOR;
+
 	// Build the default prefix too
 	DefaultPrefix = DebugPrefix + TEXT("Default");
 
@@ -356,11 +363,13 @@ struct FHotfixFileSortPredicate
 		const FString PlatformPrefix;
 		const FString ServerPrefix;
 		const FString DefaultPrefix;
+		const FString CustomConfigPrefix;
 
-		FHotfixFileNameSortPredicate(const FString& InPlatformPrefix, const FString& InServerPrefix, const FString& InDefaultPrefix) :
+		FHotfixFileNameSortPredicate(const FString& InPlatformPrefix, const FString& InServerPrefix, const FString& InDefaultPrefix, const FString& InCustomConfigPrefix) :
 			PlatformPrefix(InPlatformPrefix),
 			ServerPrefix(InServerPrefix),
-			DefaultPrefix(InDefaultPrefix)
+			DefaultPrefix(InDefaultPrefix),
+			CustomConfigPrefix(InCustomConfigPrefix)
 		{
 		}
 
@@ -389,10 +398,15 @@ struct FHotfixFileSortPredicate
 				{
 					Priority = 30;
 				}
+				// CustomConfig trumps platform
+				else if (HotfixName.StartsWith(CustomConfigPrefix))
+				{
+					Priority = 40;
+				}
 				// Other INIs listed in game override of WantsHotfixProcessing will trump all other INIs
 				else
 				{
-					Priority = 40;
+					Priority = 50;
 				}
 
 				if (!BranchVersion.IsEmpty())
@@ -428,8 +442,8 @@ struct FHotfixFileSortPredicate
 	};
 	FHotfixFileNameSortPredicate FileNameSorter;
 
-	FHotfixFileSortPredicate(const FString& InPlatformPrefix, const FString& InServerPrefix, const FString& InDefaultPrefix) :
-		FileNameSorter(InPlatformPrefix, InServerPrefix, InDefaultPrefix)
+	FHotfixFileSortPredicate(const FString& InPlatformPrefix, const FString& InServerPrefix, const FString& InDefaultPrefix, const FString& InCustomConfigPrefix) :
+		FileNameSorter(InPlatformPrefix, InServerPrefix, InDefaultPrefix, InCustomConfigPrefix)
 	{
 	}
 
@@ -460,7 +474,7 @@ void UOnlineHotfixManager::OnEnumerateFilesComplete(bool bWasSuccessful, const F
 		// Reduce the set of work to just the files that changed since last run
 		BuildHotfixFileListDeltas();
 		// Sort after filtering so that the comparison below doesn't fail to different order from the server
-		ChangedHotfixFileList.Sort<FHotfixFileSortPredicate>(FHotfixFileSortPredicate(PlatformPrefix, ServerPrefix, DefaultPrefix));
+		ChangedHotfixFileList.Sort<FHotfixFileSortPredicate>(FHotfixFileSortPredicate(PlatformPrefix, ServerPrefix, DefaultPrefix, CustomConfigPrefix));
 		// Read any changed files
 		if (ChangedHotfixFileList.Num() > 0)
 		{
@@ -879,6 +893,7 @@ bool UOnlineHotfixManager::WantsHotfixProcessing(const FCloudFileHeader& FileHea
 			bool bWantsServerHotfix = IsRunningDedicatedServer() && bIsServerHotfix;
 			bool bWantsDefaultHotfix = CloudFilename.StartsWith(DefaultPrefix);
 			bool bWantsPlatformHotfix = CloudFilename.StartsWith(PlatformPrefix);
+			bool bWantsCustomConfigPrefix = CloudFilename.StartsWith(CustomConfigPrefix);
 
 			if (bWantsPlatformHotfix)
 			{
@@ -892,8 +907,12 @@ bool UOnlineHotfixManager::WantsHotfixProcessing(const FCloudFileHeader& FileHea
 			{
 				UE_LOG(LogHotfixManager, Log, TEXT("Using default hotfix %s"), *FileHeader.FileName);
 			}
+			else if (bWantsCustomConfigPrefix)
+			{
+				UE_LOG(LogHotfixManager, Log, TEXT("Using customconfig hotfix %s"), *FileHeader.FileName);
+			}
 
-			return bWantsPlatformHotfix || bWantsServerHotfix || bWantsDefaultHotfix;
+			return bWantsPlatformHotfix || bWantsServerHotfix || bWantsDefaultHotfix || bWantsCustomConfigPrefix;
 		}
 		else
 		{
@@ -971,6 +990,10 @@ FString UOnlineHotfixManager::GetStrippedConfigFileName(const FString& IniName)
 	{
 		StrippedIniName = IniName.Right(StrippedIniName.Len() - DefaultPrefix.Len());
 	}
+	else if (StrippedIniName.StartsWith(CustomConfigPrefix))
+	{
+		StrippedIniName = IniName.Right(StrippedIniName.Len() - CustomConfigPrefix.Len());
+	}
 	else if (StrippedIniName.StartsWith(DebugPrefix))
 	{
 		StrippedIniName = IniName.Right(StrippedIniName.Len() - DebugPrefix.Len());
@@ -1121,7 +1144,7 @@ bool UOnlineHotfixManager::HotfixPakFile(const FCloudFileHeader& FileHeader)
 		}
 
 		// Sort the INIs so they are processed consistently
-		IniList.Sort<FHotfixFileSortPredicate>(FHotfixFileSortPredicate(PlatformPrefix, ServerPrefix, DefaultPrefix));
+		IniList.Sort<FHotfixFileSortPredicate>(FHotfixFileSortPredicate(PlatformPrefix, ServerPrefix, DefaultPrefix, CustomConfigPrefix));
 		// Now process the INIs in sorted order
 		for (const FString& IniName : IniList)
 		{
@@ -2357,7 +2380,7 @@ protected:
 			TestList.Add(Header);
 			Header.FileName = TEXT("DedicatedServerEngine.ini");
 			TestList.Add(Header);
-			TestList.Sort<FHotfixFileSortPredicate>(FHotfixFileSortPredicate(TEXT("PS4_"), TEXT("DedicatedServer"), TEXT("Default")));
+			TestList.Sort<FHotfixFileSortPredicate>(FHotfixFileSortPredicate(TEXT("PS4_"), TEXT("DedicatedServer"), TEXT("Default"), HOTFIX_NO_CUSTOM_CONFIG));
 
 			UE_LOG(LogHotfixManager, Log, TEXT("Hotfixing sort is:"));
 			for (const FCloudFileHeader& FileHeader : TestList)
@@ -2374,7 +2397,7 @@ protected:
 			TestList2.Add(TEXT("DefaultEngine.ini"));
 			TestList2.Add(TEXT("PS4_DefaultGame.ini"));
 			TestList2.Add(TEXT("AnotherRandom.ini"));
-			TestList2.Sort<FHotfixFileSortPredicate>(FHotfixFileSortPredicate(TEXT("PS4_"), TEXT("DedicatedServer"), TEXT("Default")));
+			TestList2.Sort<FHotfixFileSortPredicate>(FHotfixFileSortPredicate(TEXT("PS4_"), TEXT("DedicatedServer"), TEXT("Default"), HOTFIX_NO_CUSTOM_CONFIG));
 
 			UE_LOG(LogHotfixManager, Log, TEXT("Hotfixing PAK INI file sort is:"));
 			for (const FString& IniName : TestList2)
diff --git a/Engine/Plugins/Online/OnlineFramework/Source/Hotfix/Public/OnlineHotfixManager.h b/Engine/Plugins/Online/OnlineFramework/Source/Hotfix/Public/OnlineHotfixManager.h
--- a/Engine/Plugins/Online/OnlineFramework/Source/Hotfix/Public/OnlineHotfixManager.h
+++ b/Engine/Plugins/Online/OnlineFramework/Source/Hotfix/Public/OnlineHotfixManager.h
@@ -171,6 +171,8 @@ protected:
 	FString PlatformPrefix;
 	/** Used to match any server-only hotfixes */
 	FString ServerPrefix;
+	/** Used to match any custom config hotfixes */
+	FString CustomConfigPrefix;
 	/** Normally will be "Default" but could be different if we have a debug prefix */
 	FString DefaultPrefix;
 	/** Holds a chunk of string that will be swapped for Game during processing pak files (MyGame/Content/Maps -> /Game/Maps) */
