Refactored code binding to presence delegates to be hotfix friendly
Removed condition to call SetupCommonAccount in SocialUser initialization

diff --git a/Engine/Plugins/Online/OnlineFramework/Source/Party/Private/SocialToolkit.cpp b/Engine/Plugins/Online/OnlineFramework/Source/Party/Private/SocialToolkit.cpp
--- a/Engine/Plugins/Online/OnlineFramework/Source/Party/Private/SocialToolkit.cpp
+++ b/Engine/Plugins/Online/OnlineFramework/Source/Party/Private/SocialToolkit.cpp
@@ -704,24 +704,19 @@ void USocialToolkit::OnOwnerLoggedIn()
 				PartyInterface->AddOnPartyRequestToJoinRemovedDelegate_Handle(FOnPartyRequestToJoinRemovedDelegate::CreateUObject(this, &USocialToolkit::HandlePartyRequestToJoinRemoved));
 			}
 
-			if (USocialManager::ShouldReadPresenceFromOnlineServices())
+			UE::OnlineFramework::FCommonConfig CommonConfig(this);
+			if (UE::Online::IOnlineServicesPtr OnlineServices = CommonConfig.GetServices(USocialManager::GetPresenceFrameworkInstance()))
 			{
-				UE::OnlineFramework::FCommonConfig CommonConfig(this);
-				if (UE::Online::IOnlineServicesPtr OnlineServices = CommonConfig.GetServices(USocialManager::GetPresenceFrameworkInstance()))
+				if (UE::Online::IPresencePtr Presence = OnlineServices->GetPresenceInterface())
 				{
-					if (UE::Online::IPresencePtr Presence = OnlineServices->GetPresenceInterface())
-					{
-						UE::Online::FOnlineEventDelegateHandle OnPresenceUpdatedHandle = Presence->OnPresenceUpdated().Add(this, &USocialToolkit::HandleServicesPresenceUpdated, SubsystemType);
-						OnPresenceUpdatedHandles.Emplace(SubsystemType, MoveTemp(OnPresenceUpdatedHandle));
-					}
+					UE::Online::FOnlineEventDelegateHandle OnPresenceUpdatedHandle = Presence->OnPresenceUpdated().Add(this, &USocialToolkit::HandleServicesPresenceUpdated, SubsystemType);
+					OnPresenceUpdatedHandles.Emplace(SubsystemType, MoveTemp(OnPresenceUpdatedHandle));
 				}
 			}
-			else
+
+			if (IOnlinePresencePtr PresenceInterface = OSS->GetPresenceInterface())
 			{
-				if (IOnlinePresencePtr PresenceInterface = OSS->GetPresenceInterface())
-				{
-					PresenceInterface->AddOnPresenceReceivedDelegate_Handle(FOnPresenceReceivedDelegate::CreateUObject(this, &USocialToolkit::HandlePresenceReceived, SubsystemType));
-				}
+				PresenceInterface->AddOnPresenceReceivedDelegate_Handle(FOnPresenceReceivedDelegate::CreateUObject(this, &USocialToolkit::HandlePresenceReceived, SubsystemType));
 			}
 		}
 	}
@@ -794,6 +789,7 @@ void USocialToolkit::OnOwnerLoggedOut()
 		}
 	}
 
+	OnPresenceUpdatedHandles.Reset();
 	UsersBySubsystemIds.Reset();
 	AllUsers.Reset();
 
@@ -1010,60 +1006,66 @@ void USocialToolkit::HandleMapExternalIdComplete(ESocialSubsystem SubsystemType,
 
 void USocialToolkit::HandleServicesPresenceUpdated(const UE::Online::FPresenceUpdated& Event, ESocialSubsystem SubsystemType)
 {
-	using namespace UE::Online;
-	using namespace UE::OnlineFramework;
+	if (USocialManager::ShouldReadPresenceFromOnlineServices())
+	{
+		using namespace UE::Online;
+		using namespace UE::OnlineFramework;
 
-	TSharedRef<const FUserPresence> UpdatedPresence = Event.UpdatedPresence;
+		TSharedRef<const FUserPresence> UpdatedPresence = Event.UpdatedPresence;
 
-	FCommonConfig CommonConfig(this);
+		FCommonConfig CommonConfig(this);
 	
-	if (FCommonAccountPtr CommonAccountPtr = GetCommonAccountFromV2(CommonConfig, UpdatedPresence->AccountId, USocialManager::GetPresenceFrameworkInstance()))
-	{
-		// We need a services instance with an Auth adapter to make this conversion
-		if (FUniqueNetIdPtr UserId = GetV1FromCommonAccount(CommonAccountPtr.ToSharedRef(), USocialManager::GetPresenceFrameworkAdapterInstance()))
+		if (FCommonAccountPtr CommonAccountPtr = GetCommonAccountFromV2(CommonConfig, UpdatedPresence->AccountId, USocialManager::GetPresenceFrameworkInstance()))
 		{
-			QUICK_SCOPE_CYCLE_COUNTER(STAT_USocialToolkit_HandlePresenceReceived);
-			if (USocialUser* UpdatedUser = FindUser(UserId.ToSharedRef()))
+			// We need a services instance with an Auth adapter to make this conversion
+			if (FUniqueNetIdPtr UserId = GetV1FromCommonAccount(CommonAccountPtr.ToSharedRef(), USocialManager::GetPresenceFrameworkAdapterInstance()))
 			{
-				UpdatedUser->NotifyPresenceChanged(SubsystemType);
+				QUICK_SCOPE_CYCLE_COUNTER(STAT_USocialToolkit_HandlePresenceReceived);
+				if (USocialUser* UpdatedUser = FindUser(UserId.ToSharedRef()))
+				{
+					UpdatedUser->NotifyPresenceChanged(SubsystemType);
 
-				if (UpdatedUser->IsFriend())
+					if (UpdatedUser->IsFriend())
+					{
+						QueueUserDependentActionInternal(UserId.ToSharedRef(), ESocialSubsystem::Primary,
+							[this, UpdatedPresence](USocialUser& SocialUser)
+							{
+								OnFriendPresenceDidChangeV2(SocialUser, UpdatedPresence, ESocialSubsystem::Primary);
+							});
+					}
+				}
+				else if (SubsystemType == ESocialSubsystem::Platform)
 				{
-					QueueUserDependentActionInternal(UserId.ToSharedRef(), ESocialSubsystem::Primary,
-						[this, UpdatedPresence](USocialUser& SocialUser)
-						{
-							OnFriendPresenceDidChangeV2(SocialUser, UpdatedPresence, ESocialSubsystem::Primary);
-						});
+					UE_LOG(LogParty, Error, TEXT("Platform presence received for UserId [%s], but existing SocialUser could not be found."), *UserId->ToDebugString());
 				}
 			}
-			else if (SubsystemType == ESocialSubsystem::Platform)
-			{
-				UE_LOG(LogParty, Error, TEXT("Platform presence received for UserId [%s], but existing SocialUser could not be found."), *UserId->ToDebugString());
-			}
 		}
 	}
 }
 
 void USocialToolkit::HandlePresenceReceived(const FUniqueNetId& UserId, const TSharedRef<FOnlineUserPresence>& NewPresence, ESocialSubsystem SubsystemType)
 {
-	QUICK_SCOPE_CYCLE_COUNTER(STAT_USocialToolkit_HandlePresenceReceived);
-	if (USocialUser* UpdatedUser = FindUser(UserId.AsShared()))
+	if (!USocialManager::ShouldReadPresenceFromOnlineServices())
 	{
-		UpdatedUser->NotifyPresenceChanged(SubsystemType);
+		QUICK_SCOPE_CYCLE_COUNTER(STAT_USocialToolkit_HandlePresenceReceived);
+		if (USocialUser* UpdatedUser = FindUser(UserId.AsShared()))
+		{
+			UpdatedUser->NotifyPresenceChanged(SubsystemType);
 
-		if (UpdatedUser->IsFriend())
+			if (UpdatedUser->IsFriend())
+			{
+				QueueUserDependentActionInternal(UserId.AsShared(), ESocialSubsystem::Primary,
+					[this, NewPresence](USocialUser& SocialUser)
+					{
+						OnFriendPresenceDidChange(SocialUser, NewPresence, ESocialSubsystem::Primary);
+					});
+			}
+		}
+		else if (SubsystemType == ESocialSubsystem::Platform)
 		{
-			QueueUserDependentActionInternal(UserId.AsShared(), ESocialSubsystem::Primary,
-				[this, NewPresence](USocialUser& SocialUser)
-				{
-					OnFriendPresenceDidChange(SocialUser, NewPresence, ESocialSubsystem::Primary);
-				});
+			UE_LOG(LogParty, Error, TEXT("Platform presence received for UserId [%s], but existing SocialUser could not be found."), *UserId.ToDebugString());
 		}
 	}
-	else if (SubsystemType == ESocialSubsystem::Platform)
-	{
-		UE_LOG(LogParty, Error, TEXT("Platform presence received for UserId [%s], but existing SocialUser could not be found."), *UserId.ToDebugString());
-	}
 }
 
 void USocialToolkit::HandleQueryPrimaryUserIdMappingComplete(bool bWasSuccessful, const FUniqueNetId& RequestingUserId, const FString& DisplayName, const FUniqueNetId& IdentifiedUserId, const FString& Error)
diff --git a/Engine/Plugins/Online/OnlineFramework/Source/Party/Private/User/SocialUser.cpp b/Engine/Plugins/Online/OnlineFramework/Source/Party/Private/User/SocialUser.cpp
--- a/Engine/Plugins/Online/OnlineFramework/Source/Party/Private/User/SocialUser.cpp
+++ b/Engine/Plugins/Online/OnlineFramework/Source/Party/Private/User/SocialUser.cpp
@@ -2121,44 +2121,37 @@ void USocialUser::HandleQueryUserInfoComplete(ESocialSubsystem SubsystemType, bo
 		SetUserInfo(SubsystemType, UserInfo.ToSharedRef());
 	}
 
-	if (ShouldReadUserInfoFromOnlineServices())
+	SetupCommonAccount()
+	.Then([WeakThis = TWeakObjectPtr<USocialUser>(this), SubsystemType, bIsUserInfoValid = UserInfo.IsValid()](TFuture<void> Result)
 	{
-		SetupCommonAccount()
-		.Then([WeakThis = TWeakObjectPtr<USocialUser>(this), SubsystemType](TFuture<void> Result)
+		if (USocialUser* StrongThis = WeakThis.Get())
 		{
-			if (USocialUser* StrongThis = WeakThis.Get())
-			{
-				using namespace UE::OnlineFramework;
+			using namespace UE::OnlineFramework;
 
-				UE_LOG(LogParty, VeryVerbose, TEXT("SocialUser [%s] querying user info on services [%s]"), *StrongThis->ToDebugString(), ToString(SubsystemType));
+			UE_LOG(LogParty, VeryVerbose, TEXT("SocialUser [%s] querying user info on services [%s]"), *StrongThis->ToDebugString(), ToString(SubsystemType));
 
-				const FUniqueNetIdRepl UserId = StrongThis->GetUserId(SubsystemType);
+			const FUniqueNetIdRepl UserId = StrongThis->GetUserId(SubsystemType);
 
-				const FName& PresenceFrameworkInstance = USocialManager::GetPresenceFrameworkInstance();
-				FCommonConfig CommonConfig(StrongThis);
-				if (const FCommonAccountPtr CommonAccountPtr = GetCommonAccountFromV1(CommonConfig, UserId->AsShared(), GetFirstFrameworkInstanceName(UserId)))
+			const FName& PresenceFrameworkInstance = USocialManager::GetPresenceFrameworkInstance();
+			FCommonConfig CommonConfig(StrongThis);
+			if (const FCommonAccountPtr CommonAccountPtr = GetCommonAccountFromV1(CommonConfig, UserId->AsShared(), GetFirstFrameworkInstanceName(UserId)))
+			{
+				UE::Online::FAccountId AccountId = CommonAccountPtr->GetId(PresenceFrameworkInstance);
+				if (AccountId.IsValid())
 				{
-					UE::Online::FAccountId AccountId = CommonAccountPtr->GetId(PresenceFrameworkInstance);
-					if (AccountId.IsValid())
-					{
-						// No valid user info exists on this subsystem, so queue up a query for it
-						FSocialQuery_UserInfoV2::FOnQueryComplete QueryCompleteHandler = FSocialQuery_UserInfoV2::FOnQueryComplete::CreateUObject(StrongThis, &USocialUser::HandleQueryUserInfoV2Complete);
-						FSocialQueryManager::AddUserId<FSocialQuery_UserInfoV2>(StrongThis->GetOwningToolkit(), SubsystemType, AccountId, QueryCompleteHandler);
-						StrongThis->NumPendingQueries++;
+					// No valid user info exists on this subsystem, so queue up a query for it
+					FSocialQuery_UserInfoV2::FOnQueryComplete QueryCompleteHandler = FSocialQuery_UserInfoV2::FOnQueryComplete::CreateUObject(StrongThis, &USocialUser::HandleQueryUserInfoV2Complete);
+					FSocialQueryManager::AddUserId<FSocialQuery_UserInfoV2>(StrongThis->GetOwningToolkit(), SubsystemType, AccountId, QueryCompleteHandler);
+					StrongThis->NumPendingQueries++;
 
-						return;
-					}
+					return;
 				}
-
-				StrongThis->TryBroadcastInitializationComplete();
 			}
-		});
-		
-		return;	
-	}
 
-	UE_LOG(LogParty, VeryVerbose, TEXT("User [%s] finished querying user info on subsystem [%s] with result [%d]. [%d] queries still pending."), *ToDebugString(), ToString(SubsystemType), UserInfo.IsValid(), NumPendingQueries);
-	TryBroadcastInitializationComplete();
+			UE_LOG(LogParty, VeryVerbose, TEXT("User [%s] finished querying user info on subsystem [%s] with result [%d]. [%d] queries still pending."), *StrongThis->ToDebugString(), ToString(SubsystemType), bIsUserInfoValid, StrongThis->NumPendingQueries);
+			StrongThis->TryBroadcastInitializationComplete();
+		}
+	});
 }
 
 void USocialUser::HandleQueryUserInfoV2Complete(ESocialSubsystem SubsystemType, bool bWasSuccessful, const TValueOrError<TSharedRef<UE::Online::FUserInfo>, UE::Online::FOnlineError>& Result)
