Added check to avoid V2 UserInfo Query when not needed.

diff --git a/Engine/Plugins/Online/OnlineFramework/Source/Party/Private/User/SocialUser.cpp b/Engine/Plugins/Online/OnlineFramework/Source/Party/Private/User/SocialUser.cpp
--- a/Engine/Plugins/Online/OnlineFramework/Source/Party/Private/User/SocialUser.cpp
+++ b/Engine/Plugins/Online/OnlineFramework/Source/Party/Private/User/SocialUser.cpp
@@ -2126,27 +2126,31 @@ void USocialUser::HandleQueryUserInfoComplete(ESocialSubsystem SubsystemType, bo
 	{
 		if (USocialUser* StrongThis = WeakThis.Get())
 		{
-			using namespace UE::OnlineFramework;
-
-			UE_LOG(LogParty, VeryVerbose, TEXT("SocialUser [%s] querying user info on services [%s]"), *StrongThis->ToDebugString(), ToString(SubsystemType));
+			if (StrongThis->ShouldReadUserInfoFromOnlineServices())
+			{
+				using namespace UE::OnlineFramework;
 
-			const FUniqueNetIdRepl UserId = StrongThis->GetUserId(SubsystemType);
+				UE_LOG(LogParty, VeryVerbose, TEXT("SocialUser [%s] querying user info on services [%s]"), *StrongThis->ToDebugString(), ToString(SubsystemType));
 
-			const FName& PresenceFrameworkInstance = USocialManager::GetPresenceFrameworkInstance();
-			FCommonConfig CommonConfig(StrongThis);
-			if (const FCommonAccountPtr CommonAccountPtr = GetCommonAccountFromV1(CommonConfig, UserId->AsShared(), GetFirstFrameworkInstanceName(UserId)))
-			{
-				UE::Online::FAccountId AccountId = CommonAccountPtr->GetId(PresenceFrameworkInstance);
-				if (AccountId.IsValid())
+				const FUniqueNetIdRepl UserId = StrongThis->GetUserId(SubsystemType);
+				if (UserId.IsValid())
 				{
-					// No valid user info exists on this subsystem, so queue up a query for it
-					FSocialQuery_UserInfoV2::FOnQueryComplete QueryCompleteHandler = FSocialQuery_UserInfoV2::FOnQueryComplete::CreateUObject(StrongThis, &USocialUser::HandleQueryUserInfoV2Complete);
-					FSocialQueryManager::AddUserId<FSocialQuery_UserInfoV2>(StrongThis->GetOwningToolkit(), SubsystemType, AccountId, QueryCompleteHandler);
-					StrongThis->NumPendingQueries++;
+					FCommonConfig CommonConfig(StrongThis);
+					if (const FCommonAccountPtr CommonAccountPtr = GetCommonAccountFromV1(CommonConfig, UserId->AsShared(), GetFirstFrameworkInstanceName(UserId)))
+					{
+						UE::Online::FAccountId AccountId = CommonAccountPtr->GetId(USocialManager::GetPresenceFrameworkInstance());
+						if (AccountId.IsValid())
+						{
+							// No valid user info exists on this subsystem, so queue up a query for it
+							FSocialQuery_UserInfoV2::FOnQueryComplete QueryCompleteHandler = FSocialQuery_UserInfoV2::FOnQueryComplete::CreateUObject(StrongThis, &USocialUser::HandleQueryUserInfoV2Complete);
+							FSocialQueryManager::AddUserId<FSocialQuery_UserInfoV2>(StrongThis->GetOwningToolkit(), SubsystemType, AccountId, QueryCompleteHandler);
+							StrongThis->NumPendingQueries++;
 
-					return;
+							return;
+						}
+					}
 				}
-			}
+			}			
 
 			UE_LOG(LogParty, VeryVerbose, TEXT("User [%s] finished querying user info on subsystem [%s] with result [%d]. [%d] queries still pending."), *StrongThis->ToDebugString(), ToString(SubsystemType), bIsUserInfoValid, StrongThis->NumPendingQueries);
 			StrongThis->TryBroadcastInitializationComplete();
