- Re-trigger presence updates after read subsystem swap

diff --git a/Engine/Plugins/Online/OnlineFramework/Source/Party/Private/SocialToolkit.cpp b/Engine/Plugins/Online/OnlineFramework/Source/Party/Private/SocialToolkit.cpp
--- a/Engine/Plugins/Online/OnlineFramework/Source/Party/Private/SocialToolkit.cpp
+++ b/Engine/Plugins/Online/OnlineFramework/Source/Party/Private/SocialToolkit.cpp
@@ -718,6 +718,14 @@ void USocialToolkit::OnOwnerLoggedIn()
 			{
 				PresenceInterface->AddOnPresenceReceivedDelegate_Handle(FOnPresenceReceivedDelegate::CreateUObject(this, &USocialToolkit::HandlePresenceReceived, SubsystemType));
 			}
+
+			if (SubsystemType == ESocialSubsystem::Primary)
+			{
+				if (IConsoleVariable* CVar = IConsoleManager::Get().FindConsoleVariable(TEXT("onlineframework.party.shouldreadpresencefromonlineservices")))
+				{
+					CVar->OnChangedDelegate().AddUObject(this, &USocialToolkit::HandlePresenceReadCVarChanged);
+				}
+			}
 		}
 	}
 
@@ -1041,6 +1049,10 @@ void USocialToolkit::HandleServicesPresenceUpdated(const UE::Online::FPresenceUp
 			}
 		}
 	}
+	else if (SubsystemType == ESocialSubsystem::Primary)
+	{
+		NonReadPresencesReceivedV2.Emplace(Event.UpdatedPresence->AccountId);
+	}
 }
 
 void USocialToolkit::HandlePresenceReceived(const FUniqueNetId& UserId, const TSharedRef<FOnlineUserPresence>& NewPresence, ESocialSubsystem SubsystemType)
@@ -1066,6 +1078,66 @@ void USocialToolkit::HandlePresenceReceived(const FUniqueNetId& UserId, const TS
 			UE_LOG(LogParty, Error, TEXT("Platform presence received for UserId [%s], but existing SocialUser could not be found."), *UserId.ToDebugString());
 		}
 	}
+	else if (SubsystemType == ESocialSubsystem::Primary)
+	{
+		NonReadPresencesReceivedV1.Emplace(UserId.AsShared());
+	}
+}
+
+void USocialToolkit::HandlePresenceReadCVarChanged(IConsoleVariable*)
+{
+	if (USocialManager::ShouldReadPresenceFromOnlineServices())
+	{
+		using namespace UE::Online;
+		using namespace UE::OnlineFramework;
+
+		FCommonConfig CommonConfig(this);
+		IOnlineServicesPtr OnlineServices = CommonConfig.GetServices(USocialManager::GetPresenceFrameworkInstance());
+		IPresencePtr Presence = OnlineServices ? OnlineServices->GetPresenceInterface() : nullptr;
+		FUniqueNetIdRepl LocalUserIdV1 = LocalUser ? LocalUser->GetUserId(ESocialSubsystem::Primary) : nullptr;
+		const FCommonAccountPtr LocalUserCommonAccount = LocalUserIdV1.IsValid() ? GetCommonAccountFromV1(CommonConfig, LocalUserIdV1.GetV1(), GetFirstFrameworkInstanceName(*LocalUserIdV1)) : nullptr;
+		const FAccountId LocalUserAccountId = LocalUserCommonAccount ? LocalUserCommonAccount->GetId(USocialManager::GetPresenceFrameworkInstance()) : FAccountId();
+		if (Presence && LocalUserAccountId.IsValid())
+		{
+			const TSet<FAccountId> PreviouslyUpdatedPresences = MoveTemp(NonReadPresencesReceivedV2);
+			for (const FAccountId TargetAccountId : PreviouslyUpdatedPresences)
+			{
+				TOnlineResult<FGetCachedPresence> Result = Presence->GetCachedPresence({
+					.LocalAccountId = LocalUserAccountId,
+					.TargetAccountId = TargetAccountId
+				});
+
+				if (Result.IsOk())
+				{
+					HandleServicesPresenceUpdated({
+						.LocalAccountId = LocalUserAccountId,
+						.UpdatedPresence = Result.GetOkValue().Presence
+					}, ESocialSubsystem::Primary);
+				}
+			}
+		}
+	}
+	else 
+	{
+		if (IOnlineSubsystem* Oss = GetSocialOss(ESocialSubsystem::Primary))
+		{
+			IOnlinePresencePtr Presence = Oss->GetPresenceInterface();
+			FUniqueNetIdRepl LocalUserId = LocalUser ? LocalUser->GetUserId(ESocialSubsystem::Primary) : nullptr;
+			if (Presence && LocalUserId.IsValid())
+			{
+				const FUniqueNetIdSet PreviouslyUpdatedNetIds = MoveTemp(NonReadPresencesReceivedV1);
+				for (const FUniqueNetIdRef& TargetNetId : PreviouslyUpdatedNetIds)
+				{
+					TSharedPtr<FOnlineUserPresence> UserPresence;
+					EOnlineCachedResult::Type Result = Presence->GetCachedPresenceForApp(*LocalUserId, *TargetNetId, Oss->GetAppId(), UserPresence);
+					if (UserPresence)
+					{
+						HandlePresenceReceived(*TargetNetId, UserPresence.ToSharedRef(), ESocialSubsystem::Primary);
+					}
+				}
+			}
+		}
+	}
 }
 
 void USocialToolkit::HandleQueryPrimaryUserIdMappingComplete(bool bWasSuccessful, const FUniqueNetId& RequestingUserId, const FString& DisplayName, const FUniqueNetId& IdentifiedUserId, const FString& Error)
diff --git a/Engine/Plugins/Online/OnlineFramework/Source/Party/Public/SocialToolkit.h b/Engine/Plugins/Online/OnlineFramework/Source/Party/Public/SocialToolkit.h
--- a/Engine/Plugins/Online/OnlineFramework/Source/Party/Public/SocialToolkit.h
+++ b/Engine/Plugins/Online/OnlineFramework/Source/Party/Public/SocialToolkit.h
@@ -273,6 +273,7 @@ private:	// Handlers
 	UE_API void HandleMapExternalIdComplete(ESocialSubsystem SubsystemType, bool bWasSuccessful, const FUniqueNetIdRepl& MappedPrimaryId, FUniqueNetIdRepl ExternalId, TFunction<void(USocialUser&)> UserActionFunc, bool bExecutePostInit);
 	UE_API void HandlePresenceReceived(const FUniqueNetId& UserId, const TSharedRef<FOnlineUserPresence>& NewPresence, ESocialSubsystem SubsystemType);
 	UE_API void HandleServicesPresenceUpdated(const UE::Online::FPresenceUpdated& Event, ESocialSubsystem SubsystemType);
+	void HandlePresenceReadCVarChanged(IConsoleVariable*);
 
 	UE_API void HandleFriendInviteReceived(const FUniqueNetId& LocalUserId, const FUniqueNetId& SenderId, ESocialSubsystem SubsystemType);
 	UE_API void HandleFriendInviteAccepted(const FUniqueNetId& LocalUserId, const FUniqueNetId& NewFriendId, ESocialSubsystem SubsystemType);
@@ -339,6 +340,9 @@ private:
 	UPROPERTY()
 	TWeakObjectPtr<ULocalPlayer> LocalPlayerOwner = nullptr;
 
+	TSet<UE::Online::FAccountId> NonReadPresencesReceivedV2;
+	FUniqueNetIdSet NonReadPresencesReceivedV1;
+
 	TSet<IOnlinePartyJoinInfoConstRef> PartyInvitations;
 	mutable TArray<TWeakPtr<FSocialUserList>> CachedSocialUserLists;
 
