- Some additions from code review

#Synced-CL 48072565

diff --git a/Engine/Plugins/Online/OnlineFramework/Source/Party/Private/SocialToolkit.cpp b/Engine/Plugins/Online/OnlineFramework/Source/Party/Private/SocialToolkit.cpp
--- a/Engine/Plugins/Online/OnlineFramework/Source/Party/Private/SocialToolkit.cpp
+++ b/Engine/Plugins/Online/OnlineFramework/Source/Party/Private/SocialToolkit.cpp
@@ -1100,6 +1100,7 @@ void USocialToolkit::HandlePresenceReadCVarChanged(IConsoleVariable*)
 		if (Presence && LocalUserAccountId.IsValid())
 		{
 			const TSet<FAccountId> PreviouslyUpdatedPresences = MoveTemp(NonReadPresencesReceivedV2);
+			NonReadPresencesReceivedV2.Empty();
 			for (const FAccountId TargetAccountId : PreviouslyUpdatedPresences)
 			{
 				TOnlineResult<FGetCachedPresence> Result = Presence->GetCachedPresence({
@@ -1126,6 +1127,7 @@ void USocialToolkit::HandlePresenceReadCVarChanged(IConsoleVariable*)
 			if (Presence && LocalUserId.IsValid())
 			{
 				const FUniqueNetIdSet PreviouslyUpdatedNetIds = MoveTemp(NonReadPresencesReceivedV1);
+				NonReadPresencesReceivedV1.Empty();
 				for (const FUniqueNetIdRef& TargetNetId : PreviouslyUpdatedNetIds)
 				{
 					TSharedPtr<FOnlineUserPresence> UserPresence;
