Added checks to restrict OnlineServices presence functionality only to Primary social subsystem

#Synced-CL 48271216

diff --git a/Engine/Plugins/Online/OnlineFramework/Source/Party/Private/SocialToolkit.cpp b/Engine/Plugins/Online/OnlineFramework/Source/Party/Private/SocialToolkit.cpp
--- a/Engine/Plugins/Online/OnlineFramework/Source/Party/Private/SocialToolkit.cpp
+++ b/Engine/Plugins/Online/OnlineFramework/Source/Party/Private/SocialToolkit.cpp
@@ -704,16 +704,6 @@ void USocialToolkit::OnOwnerLoggedIn()
 				PartyInterface->AddOnPartyRequestToJoinRemovedDelegate_Handle(FOnPartyRequestToJoinRemovedDelegate::CreateUObject(this, &USocialToolkit::HandlePartyRequestToJoinRemoved));
 			}
 
-			UE::OnlineFramework::FCommonConfig CommonConfig(this);
-			if (UE::Online::IOnlineServicesPtr OnlineServices = CommonConfig.GetServices(USocialManager::GetPresenceFrameworkInstance()))
-			{
-				if (UE::Online::IPresencePtr Presence = OnlineServices->GetPresenceInterface())
-				{
-					UE::Online::FOnlineEventDelegateHandle OnPresenceUpdatedHandle = Presence->OnPresenceUpdated().Add(this, &USocialToolkit::HandleServicesPresenceUpdated, SubsystemType);
-					OnPresenceUpdatedHandles.Emplace(SubsystemType, MoveTemp(OnPresenceUpdatedHandle));
-				}
-			}
-
 			if (IOnlinePresencePtr PresenceInterface = OSS->GetPresenceInterface())
 			{
 				PresenceInterface->AddOnPresenceReceivedDelegate_Handle(FOnPresenceReceivedDelegate::CreateUObject(this, &USocialToolkit::HandlePresenceReceived, SubsystemType));
@@ -721,6 +711,16 @@ void USocialToolkit::OnOwnerLoggedIn()
 
 			if (SubsystemType == ESocialSubsystem::Primary)
 			{
+				UE::OnlineFramework::FCommonConfig CommonConfig(this);
+				if (UE::Online::IOnlineServicesPtr OnlineServices = CommonConfig.GetServices(USocialManager::GetPresenceFrameworkInstance()))
+				{
+					if (UE::Online::IPresencePtr Presence = OnlineServices->GetPresenceInterface())
+					{
+						UE::Online::FOnlineEventDelegateHandle OnPresenceUpdatedHandle = Presence->OnPresenceUpdated().Add(this, &USocialToolkit::HandleServicesPresenceUpdated, SubsystemType);
+						OnPresenceUpdatedHandles.Emplace(SubsystemType, MoveTemp(OnPresenceUpdatedHandle));
+					}
+				}
+
 				if (IConsoleVariable* CVar = IConsoleManager::Get().FindConsoleVariable(TEXT("onlineframework.party.shouldreadpresencefromonlineservices")))
 				{
 					CVar->OnChangedDelegate().AddUObject(this, &USocialToolkit::HandlePresenceReadCVarChanged);
@@ -1049,7 +1049,7 @@ void USocialToolkit::HandleServicesPresenceUpdated(const UE::Online::FPresenceUp
 			}
 		}
 	}
-	else if (SubsystemType == ESocialSubsystem::Primary)
+	else
 	{
 		NonReadPresencesReceivedV2.Emplace(Event.UpdatedPresence->AccountId);
 	}
@@ -1057,7 +1057,7 @@ void USocialToolkit::HandleServicesPresenceUpdated(const UE::Online::FPresenceUp
 
 void USocialToolkit::HandlePresenceReceived(const FUniqueNetId& UserId, const TSharedRef<FOnlineUserPresence>& NewPresence, ESocialSubsystem SubsystemType)
 {
-	if (!USocialManager::ShouldReadPresenceFromOnlineServices())
+	if (!USocialManager::ShouldReadPresenceFromOnlineServices() || SubsystemType == ESocialSubsystem::Platform)
 	{
 		QUICK_SCOPE_CYCLE_COUNTER(STAT_USocialToolkit_HandlePresenceReceived);
 		if (USocialUser* UpdatedUser = FindUser(UserId.AsShared()))
@@ -1078,7 +1078,7 @@ void USocialToolkit::HandlePresenceReceived(const FUniqueNetId& UserId, const TS
 			UE_LOG(LogParty, Error, TEXT("Platform presence received for UserId [%s], but existing SocialUser could not be found."), *UserId.ToDebugString());
 		}
 	}
-	else if (SubsystemType == ESocialSubsystem::Primary)
+	else
 	{
 		NonReadPresencesReceivedV1.Emplace(UserId.AsShared());
 	}
diff --git a/Engine/Plugins/Online/OnlineFramework/Source/Party/Private/User/SocialUser.cpp b/Engine/Plugins/Online/OnlineFramework/Source/Party/Private/User/SocialUser.cpp
--- a/Engine/Plugins/Online/OnlineFramework/Source/Party/Private/User/SocialUser.cpp
+++ b/Engine/Plugins/Online/OnlineFramework/Source/Party/Private/User/SocialUser.cpp
@@ -2130,7 +2130,7 @@ void USocialUser::HandleQueryUserInfoComplete(ESocialSubsystem SubsystemType, bo
 	{
 		if (USocialUser* StrongThis = WeakThis.Get())
 		{
-			if (StrongThis->ShouldReadUserInfoFromOnlineServices())
+			if (StrongThis->ShouldReadUserInfoFromOnlineServices() && SubsystemType == ESocialSubsystem::Primary)
 			{
 				using namespace UE::OnlineFramework;
 
diff --git a/Engine/Plugins/Online/OnlineFramework/Source/Party/Private/User/SocialUserList.cpp b/Engine/Plugins/Online/OnlineFramework/Source/Party/Private/User/SocialUserList.cpp
--- a/Engine/Plugins/Online/OnlineFramework/Source/Party/Private/User/SocialUserList.cpp
+++ b/Engine/Plugins/Online/OnlineFramework/Source/Party/Private/User/SocialUserList.cpp
@@ -416,7 +416,7 @@ bool FSocialUserList::EvaluateUserPresence(const USocialUser& User, ESocialSubsy
 		bool bIsPlayingThisGame = false;
 		bool bInSameParty = false;
 
-		if (USocialManager::ShouldReadPresenceFromOnlineServices())
+		if (USocialManager::ShouldReadPresenceFromOnlineServices() && SubsystemType == ESocialSubsystem::Primary)
 		{
 			using namespace UE::Online;
 
