Added missing checks for platform presence to be always processed via OSS

#Synced-CL 48271216

diff --git a/Engine/Plugins/Online/OnlineFramework/Source/Party/Private/User/SocialUser.cpp b/Engine/Plugins/Online/OnlineFramework/Source/Party/Private/User/SocialUser.cpp
--- a/Engine/Plugins/Online/OnlineFramework/Source/Party/Private/User/SocialUser.cpp
+++ b/Engine/Plugins/Online/OnlineFramework/Source/Party/Private/User/SocialUser.cpp
@@ -8,6 +8,7 @@
 
 #include "Online/OnlineErrorDefinitions.h"
 #include "Online/OnlineServices.h"
+#include "Online/PresenceOSSAdapter.h"
 #include "OnlineFramework/CommonAccount.h"
 #include "OnlineFramework/CommonAccountUtils.h"
 #include "OnlineFramework/CommonConfig.h"
@@ -579,34 +580,53 @@ UE::Online::EUserPresenceStatus USocialUser::GetOnlineStatusV2() const
 			// Get the most "present" status available on any of the associated platforms
 			for (const TPair<ESocialSubsystem, FSubsystemUserInfo>& SubsystemInfoPair : SubsystemInfoByType)
 			{
-				const FUniqueNetIdRepl& LocalUserNetId = GetOwningToolkit().GetLocalUserNetId(SubsystemInfoPair.Key);
-				if(LocalUserNetId.IsValid())
+				if (SubsystemInfoPair.Key == ESocialSubsystem::Primary)
 				{
-					if (const FCommonAccountPtr LocalUserCommonAccountPtr = GetCommonAccountFromV1(CommonConfig, LocalUserNetId->AsShared(), GetFirstFrameworkInstanceName(LocalUserNetId)))
+					const FUniqueNetIdRepl& LocalUserNetId = GetOwningToolkit().GetLocalUserNetId(SubsystemInfoPair.Key);
+					if(LocalUserNetId.IsValid())
 					{
-						FAccountId LocalUserAccountId = LocalUserCommonAccountPtr->GetId(PresenceFrameworkInstance);
-
-						const FUniqueNetIdRepl& UserId = SubsystemInfoPair.Value.GetUserId();
-
-						if (const FCommonAccountPtr CommonAccountPtr = GetCommonAccountFromV1(CommonConfig, UserId->AsShared(), GetFirstFrameworkInstanceName(UserId)))
+						if (const FCommonAccountPtr LocalUserCommonAccountPtr = GetCommonAccountFromV1(CommonConfig, LocalUserNetId->AsShared(), GetFirstFrameworkInstanceName(LocalUserNetId)))
 						{
-							FAccountId AccountId = CommonAccountPtr->GetId(PresenceFrameworkInstance);
+							FAccountId LocalUserAccountId = LocalUserCommonAccountPtr->GetId(PresenceFrameworkInstance);
+
+							const FUniqueNetIdRepl& UserId = SubsystemInfoPair.Value.GetUserId();
 
-							TOnlineResult<FGetCachedPresence> GetCachedPresenceResult = Presence->GetCachedPresence({ LocalUserAccountId, AccountId });
-							if (GetCachedPresenceResult.IsOk())
+							if (const FCommonAccountPtr CommonAccountPtr = GetCommonAccountFromV1(CommonConfig, UserId->AsShared(), GetFirstFrameworkInstanceName(UserId)))
 							{
-								EUserPresenceStatus CachedPresenceStatus = GetCachedPresenceResult.GetOkValue().Presence->Status;
+								FAccountId AccountId = CommonAccountPtr->GetId(PresenceFrameworkInstance);
 
-								if (Result == EUserPresenceStatus::Offline || CachedPresenceStatus == EUserPresenceStatus::Online ||
-								(CachedPresenceStatus == EUserPresenceStatus::Away && Result != EUserPresenceStatus::Online))
+								TOnlineResult<FGetCachedPresence> GetCachedPresenceResult = Presence->GetCachedPresence({ LocalUserAccountId, AccountId });
+								if (GetCachedPresenceResult.IsOk())
 								{
-									// Either the best we have is offline, or the new one is either online or away (if necessary we can get into the weeds of prioritizing the other states)
-									Result = CachedPresenceStatus;
+									EUserPresenceStatus CachedPresenceStatus = GetCachedPresenceResult.GetOkValue().Presence->Status;
+
+									if (Result == EUserPresenceStatus::Offline || CachedPresenceStatus == EUserPresenceStatus::Online ||
+									(CachedPresenceStatus == EUserPresenceStatus::Away && Result != EUserPresenceStatus::Online))
+									{
+										// Either the best we have is offline, or the new one is either online or away (if necessary we can get into the weeds of prioritizing the other states)
+										Result = CachedPresenceStatus;
+									}
 								}
 							}
 						}
 					}
 				}
+				else if (SubsystemInfoPair.Key == ESocialSubsystem::Platform) // For Platform presence, we'll get it from V1 and convert it
+				{
+					EOnlinePresenceState::Type OnlineStatus = EOnlinePresenceState::Offline;
+
+					if (const FOnlineUserPresence* PresenceInfo = SubsystemInfoPair.Value.GetPresenceInfo())
+					{
+						if (OnlineStatus == EOnlinePresenceState::Offline || PresenceInfo->Status.State == EOnlinePresenceState::Online ||
+							(PresenceInfo->Status.State == EOnlinePresenceState::Away && OnlineStatus != EOnlinePresenceState::Online))
+						{
+							// Either the best we have is offline, or the new one is either online or away (if necessary we can get into the weeds of prioritizing the other states)
+							OnlineStatus = PresenceInfo->Status.State;
+						}
+					}
+
+					Result = UE::Online::FPresenceOSSAdapter::StatusV1toV2(OnlineStatus);
+				}
 			}
 		}
 	}
@@ -1357,7 +1377,7 @@ bool USocialUser::IsOnline() const
 	// If any presence says we're online, count it (note also that only friends have presence info, so non-friends will always count as offline)
 	for (const TPair<ESocialSubsystem, FSubsystemUserInfo>& SubsystemInfoPair : SubsystemInfoByType)
 	{
-		if (USocialManager::ShouldReadPresenceFromOnlineServices() || ShouldReadUserInfoFromOnlineServices())
+		if ((USocialManager::ShouldReadPresenceFromOnlineServices() || ShouldReadUserInfoFromOnlineServices()) && SubsystemInfoPair.Key == ESocialSubsystem::Primary)
 		{
 			using namespace UE::Online;
 			using namespace UE::OnlineFramework;
@@ -1418,7 +1438,7 @@ bool USocialUser::IsPlayingThisGame() const
 {
 	for (const TPair<ESocialSubsystem, FSubsystemUserInfo>& SubsystemInfoPair : SubsystemInfoByType)
 	{
-		if (USocialManager::ShouldReadPresenceFromOnlineServices() || ShouldReadUserInfoFromOnlineServices())
+		if ((USocialManager::ShouldReadPresenceFromOnlineServices() || ShouldReadUserInfoFromOnlineServices()) && SubsystemInfoPair.Key == ESocialSubsystem::Primary)
 		{
 			using namespace UE::Online;
 			using namespace UE::OnlineFramework;
