Refactor to use current user platform for EOS Display Name retrieval instead of EOS Presence platform type

#Synced-CL 48360866

diff --git a/Engine/Plugins/Online/OnlineFramework/Source/Party/Private/User/SocialUser.cpp b/Engine/Plugins/Online/OnlineFramework/Source/Party/Private/User/SocialUser.cpp
--- a/Engine/Plugins/Online/OnlineFramework/Source/Party/Private/User/SocialUser.cpp
+++ b/Engine/Plugins/Online/OnlineFramework/Source/Party/Private/User/SocialUser.cpp
@@ -711,17 +711,41 @@ bool USocialUser::ShouldReadUserInfoFromOnlineServices() const
 	return false;
 }
 
+UE::Online::EOnlinePlatformType USocialUser::GetDisplayNameOnlinePlatformTypeFromUserPlatform(const FString& UserCurrentPlatformStr) const
+{
+	using namespace UE::Online;
+
+	if (UserCurrentPlatformStr == TEXT("PS5") || UserCurrentPlatformStr == TEXT("PSN"))
+	{
+		return EOnlinePlatformType::PSN;
+	}
+	else if (UserCurrentPlatformStr == TEXT("XSX") || UserCurrentPlatformStr == TEXT("XBL"))
+	{
+		return EOnlinePlatformType::XBL;
+	}
+	else if (UserCurrentPlatformStr == TEXT("SWT2") || UserCurrentPlatformStr == TEXT("SWT"))
+	{
+		return EOnlinePlatformType::Nintendo;
+	}
+	else
+	{
+		// Intentional fall-through for all other platforms
+		return EOnlinePlatformType::Epic;
+	}
+}
+
 FString USocialUser::GetDisplayName() const
 {
 	FString DisplayName;
 
 	if(const FSubsystemUserInfo* SubsystemInfo = SubsystemInfoByType.Find(ESocialSubsystem::Primary))
 	{
+		const FUserPlatform& UserCurrentPlatform = GetCurrentPlatform();
+
 		// We always retrieve the V1 name during the transition period, for the redundancy check
 		if (TSharedPtr<FOnlineUser> PrimaryUserInfo = SubsystemInfo->UserInfo.Pin())
 		{
 			// The primary user info has knowledge of display names on all linked accounts
-			const FUserPlatform& UserCurrentPlatform = GetCurrentPlatform();
 			DisplayName = PrimaryUserInfo->GetDisplayName(UserCurrentPlatform.ToString());
 		}
 		else
@@ -757,16 +781,7 @@ FString USocialUser::GetDisplayName() const
 
 							if (ShouldSpecifyDisplayNameServicesProvider(SubsystemInfo))
 							{
-								// We try to get the user platform from presence
-								EOnlinePlatformType OnlinePlatformType = EOnlinePlatformType::Unknown;
-								if (IPresencePtr Presence = OnlineServices->GetPresenceInterface())
-								{
-									TOnlineResult<FGetCachedPresence> GetCachedPresenceResult = Presence->GetCachedPresence({ LocalUserAccountId, SocialUserAccountId });
-									if (GetCachedPresenceResult.IsOk())
-									{
-										OnlinePlatformType = GetCachedPresenceResult.GetOkValue().Presence->PlatformType;
-									}
-								}
+								EOnlinePlatformType OnlinePlatformType = GetDisplayNameOnlinePlatformTypeFromUserPlatform(UserCurrentPlatform.ToString());
 								
 								TOnlineResult<FGetUserInfo> GetUserInfoResult = UserInfo->GetUserInfo({ LocalUserAccountId, SocialUserAccountId, OnlinePlatformType });
 								if (GetUserInfoResult.IsOk())
diff --git a/Engine/Plugins/Online/OnlineFramework/Source/Party/Public/User/SocialUser.h b/Engine/Plugins/Online/OnlineFramework/Source/Party/Public/User/SocialUser.h
--- a/Engine/Plugins/Online/OnlineFramework/Source/Party/Public/User/SocialUser.h
+++ b/Engine/Plugins/Online/OnlineFramework/Source/Party/Public/User/SocialUser.h
@@ -253,6 +253,7 @@ protected:
 	UE_API virtual TFuture<void> SetupCommonAccount();
 	// This method should be extended in derived classes to make any additional subsystem user info checks needed that might modify the services provider used in GetDisplayName
 	UE_API virtual bool ShouldSpecifyDisplayNameServicesProvider(const FSubsystemUserInfo* SubsystemInfo) const;
+	UE_API virtual UE::Online::EOnlinePlatformType GetDisplayNameOnlinePlatformTypeFromUserPlatform(const FString& UserCurrentPlatformStr) const;
 
 	UE_API bool ShouldReadUserInfoFromOnlineServices() const;
 
