Fix party join flow getting stuck if FindSession fails
#Synced-CL 48396418

diff --git a/Engine/Plugins/Online/OnlineFramework/Source/Party/Private/Party/PartyPlatformSessionMonitor.cpp b/Engine/Plugins/Online/OnlineFramework/Source/Party/Private/Party/PartyPlatformSessionMonitor.cpp
--- a/Engine/Plugins/Online/OnlineFramework/Source/Party/Private/Party/PartyPlatformSessionMonitor.cpp
+++ b/Engine/Plugins/Online/OnlineFramework/Source/Party/Private/Party/PartyPlatformSessionMonitor.cpp
@@ -259,7 +259,11 @@ void FPartyPlatformSessionManager::HandleFindSessionByIdComplete(int32 LocalUser
 {
 	UE_LOG(LogParty, Verbose, TEXT("FPartyPlatformSessionManager completed attempt to find platform session [%s] of user [%s] by SessionId with result [%d]"), *SessionId, *SessionOwnerId.ToDebugString(), (int32)bWasSuccessful);
 
-	if (!bWasSuccessful || !FoundSession.IsSessionInfoValid())
+	if (bWasSuccessful && FoundSession.IsSessionInfoValid())
+	{
+		ProcessCompletedSessionSearch(LocalUserNum, true, FoundSession, SessionId, SessionOwnerId, OnAttemptComplete);
+	}
+	else
 	{
 #if PARTY_PLATFORM_SESSIONS_PSN
 		//@todo DanH: Obviously remove all of this asap - we need the PSN OSS to be able to get updated presence info for a single user without querying the whole frigging list
@@ -282,10 +286,10 @@ void FPartyPlatformSessionManager::HandleFindSessionByIdComplete(int32 LocalUser
 			SessionInterface->AddOnFindFriendSessionCompleteDelegate_Handle(LocalUserNum, FOnFindFriendSessionCompleteDelegate::CreateSP(this, &FPartyPlatformSessionManager::HandleFindFriendSessionsComplete, SessionId, SessionOwnerId, OnAttemptComplete));
 			SessionInterface->FindFriendSession(LocalUserNum, *SessionOwnerId);
 		}
-	}
-	else
-	{
-		ProcessCompletedSessionSearch(LocalUserNum, true, FoundSession, SessionId, SessionOwnerId, OnAttemptComplete);
+		else
+		{
+			OnAttemptComplete.ExecuteIfBound(false, FoundSession);
+		}
 	}
 }
 
