Added id validity checks to usage in SocialUser

#Synced-CL 48438952

diff --git a/Engine/Plugins/Online/OnlineFramework/Source/Party/Private/User/SocialUser.cpp b/Engine/Plugins/Online/OnlineFramework/Source/Party/Private/User/SocialUser.cpp
--- a/Engine/Plugins/Online/OnlineFramework/Source/Party/Private/User/SocialUser.cpp
+++ b/Engine/Plugins/Online/OnlineFramework/Source/Party/Private/User/SocialUser.cpp
@@ -590,21 +590,23 @@ UE::Online::EUserPresenceStatus USocialUser::GetOnlineStatusV2() const
 							FAccountId LocalUserAccountId = LocalUserCommonAccountPtr->GetId(PresenceFrameworkInstance);
 
 							const FUniqueNetIdRepl& UserId = SubsystemInfoPair.Value.GetUserId();
-
-							if (const FCommonAccountPtr CommonAccountPtr = GetCommonAccountFromV1(CommonConfig, UserId->AsShared(), GetFirstFrameworkInstanceName(UserId)))
+							if (UserId.IsValid())
 							{
-								FAccountId AccountId = CommonAccountPtr->GetId(PresenceFrameworkInstance);
-
-								TOnlineResult<FGetCachedPresence> GetCachedPresenceResult = Presence->GetCachedPresence({ LocalUserAccountId, AccountId });
-								if (GetCachedPresenceResult.IsOk())
+								if (const FCommonAccountPtr CommonAccountPtr = GetCommonAccountFromV1(CommonConfig, UserId->AsShared(), GetFirstFrameworkInstanceName(UserId)))
 								{
-									EUserPresenceStatus CachedPresenceStatus = GetCachedPresenceResult.GetOkValue().Presence->Status;
+									FAccountId AccountId = CommonAccountPtr->GetId(PresenceFrameworkInstance);
 
-									if (Result == EUserPresenceStatus::Offline || CachedPresenceStatus == EUserPresenceStatus::Online ||
-									(CachedPresenceStatus == EUserPresenceStatus::Away && Result != EUserPresenceStatus::Online))
+									TOnlineResult<FGetCachedPresence> GetCachedPresenceResult = Presence->GetCachedPresence({ LocalUserAccountId, AccountId });
+									if (GetCachedPresenceResult.IsOk())
 									{
-										// Either the best we have is offline, or the new one is either online or away (if necessary we can get into the weeds of prioritizing the other states)
-										Result = CachedPresenceStatus;
+										EUserPresenceStatus CachedPresenceStatus = GetCachedPresenceResult.GetOkValue().Presence->Status;
+
+										if (Result == EUserPresenceStatus::Offline || CachedPresenceStatus == EUserPresenceStatus::Online ||
+										(CachedPresenceStatus == EUserPresenceStatus::Away && Result != EUserPresenceStatus::Online))
+										{
+											// Either the best we have is offline, or the new one is either online or away (if necessary we can get into the weeds of prioritizing the other states)
+											Result = CachedPresenceStatus;
+										}
 									}
 								}
 							}
@@ -763,57 +765,62 @@ FString USocialUser::GetDisplayName() const
 			FCommonConfig CommonConfig(this);
 
 			const FUniqueNetIdRepl& LocalUserNetId = GetOwningToolkit().GetLocalUserNetId(ESocialSubsystem::Primary);
-			if (const FCommonAccountPtr LocalUserCommonAccountPtr = GetCommonAccountFromV1(CommonConfig, LocalUserNetId->AsShared(), GetFirstFrameworkInstanceName(LocalUserNetId)))
+			if (LocalUserNetId.IsValid())
 			{
-				FAccountId LocalUserAccountId = LocalUserCommonAccountPtr->GetId(PresenceFrameworkInstance);
-
-				const FUniqueNetIdRepl& SocialUserId = SubsystemInfo->GetUserId();
-
-				if (const FCommonAccountPtr SocialUserCommonAccountPtr = GetCommonAccountFromV1(CommonConfig, SocialUserId->AsShared(), GetFirstFrameworkInstanceName(SocialUserId)))
+				if (const FCommonAccountPtr LocalUserCommonAccountPtr = GetCommonAccountFromV1(CommonConfig, LocalUserNetId->AsShared(), GetFirstFrameworkInstanceName(LocalUserNetId)))
 				{
-					FAccountId SocialUserAccountId = SocialUserCommonAccountPtr->GetId(PresenceFrameworkInstance);
+					FAccountId LocalUserAccountId = LocalUserCommonAccountPtr->GetId(PresenceFrameworkInstance);
 
-					if (IOnlineServicesPtr OnlineServices = CommonConfig.GetServices(PresenceFrameworkInstance))
+					const FUniqueNetIdRepl& SocialUserId = SubsystemInfo->GetUserId();
+					if (SocialUserId.IsValid())
 					{
-						if (IUserInfoPtr UserInfo = OnlineServices->GetUserInfoInterface())
+						if (const FCommonAccountPtr SocialUserCommonAccountPtr = GetCommonAccountFromV1(CommonConfig, SocialUserId->AsShared(), GetFirstFrameworkInstanceName(SocialUserId)))
 						{
-							FString DisplayNameV2;
+							FAccountId SocialUserAccountId = SocialUserCommonAccountPtr->GetId(PresenceFrameworkInstance);
 
-							if (ShouldSpecifyDisplayNameServicesProvider(SubsystemInfo))
+							if (IOnlineServicesPtr OnlineServices = CommonConfig.GetServices(PresenceFrameworkInstance))
 							{
-								EOnlinePlatformType OnlinePlatformType = GetDisplayNameOnlinePlatformTypeFromUserPlatform(UserCurrentPlatform.ToString());
-								
-								TOnlineResult<FGetUserInfo> GetUserInfoResult = UserInfo->GetUserInfo({ LocalUserAccountId, SocialUserAccountId, OnlinePlatformType });
-								if (GetUserInfoResult.IsOk())
+								if (IUserInfoPtr UserInfo = OnlineServices->GetUserInfoInterface())
 								{
-									DisplayNameV2 = GetUserInfoResult.GetOkValue().UserInfo->DisplayName;
-								}
-							}
+									FString DisplayNameV2;
 
-							// If we don't get a valid result after specifying the services provider for DisplayName, we'll call again to iterate through Epic and all available platforms
-							if (DisplayNameV2.IsEmpty())
-							{
-								TOnlineResult<FGetUserInfo> GetUserInfoResult = UserInfo->GetUserInfo({ LocalUserAccountId, SocialUserAccountId });
-								if (GetUserInfoResult.IsOk())
-								{
-									DisplayNameV2 = GetUserInfoResult.GetOkValue().UserInfo->DisplayName;
-								}
-								else
-								{
-									// We don't have primary user info (so we're not even initialized yet!), but a good-faith effort is to see if we have a platform name
-									DisplayNameV2 = GetDisplayName(ESocialSubsystem::Platform);
-								}
-							}
+									if (ShouldSpecifyDisplayNameServicesProvider(SubsystemInfo))
+									{
+										EOnlinePlatformType OnlinePlatformType = GetDisplayNameOnlinePlatformTypeFromUserPlatform(UserCurrentPlatform.ToString());
+								
+										TOnlineResult<FGetUserInfo> GetUserInfoResult = UserInfo->GetUserInfo({ LocalUserAccountId, SocialUserAccountId, OnlinePlatformType });
+										if (GetUserInfoResult.IsOk())
+										{
+											DisplayNameV2 = GetUserInfoResult.GetOkValue().UserInfo->DisplayName;
+										}
+									}
+
+									// If we don't get a valid result after specifying the services provider for DisplayName, we'll call again to iterate through Epic and all available platforms
+									if (DisplayNameV2.IsEmpty())
+									{
+										TOnlineResult<FGetUserInfo> GetUserInfoResult = UserInfo->GetUserInfo({ LocalUserAccountId, SocialUserAccountId });
+										if (GetUserInfoResult.IsOk())
+										{
+											DisplayNameV2 = GetUserInfoResult.GetOkValue().UserInfo->DisplayName;
+										}
+										else
+										{
+											// We don't have primary user info (so we're not even initialized yet!), but a good-faith effort is to see if we have a platform name
+											DisplayNameV2 = GetDisplayName(ESocialSubsystem::Platform);
+										}
+									}
 #if !UE_BUILD_SHIPPING
-							// Temporary redundancy check
-							if (DisplayName != DisplayNameV2)
-							{
-								UE_LOG(LogParty, Warning, TEXT("DisplayName disparity: DisplayNameV1: %s. DisplayNameV2: %s"), *DisplayName, *DisplayNameV2);
-							}
+									// Temporary redundancy check
+									if (DisplayName != DisplayNameV2)
+									{
+										UE_LOG(LogParty, Warning, TEXT("DisplayName disparity: DisplayNameV1: %s. DisplayNameV2: %s"), *DisplayName, *DisplayNameV2);
+									}
 #endif
-							if (!DisplayNameV2.IsEmpty())
-							{
-								DisplayName = DisplayNameV2;
+									if (!DisplayNameV2.IsEmpty())
+									{
+										DisplayName = DisplayNameV2;
+									}
+								}
 							}
 						}
 					}
@@ -844,34 +851,39 @@ FString USocialUser::GetDisplayName(ESocialSubsystem SubsystemType) const
 				FCommonConfig CommonConfig(this);
 
 				const FUniqueNetIdRepl& LocalUserNetId = GetOwningToolkit().GetLocalUserNetId(ESocialSubsystem::Primary);
-				if (const FCommonAccountPtr LocalUserCommonAccountPtr = GetCommonAccountFromV1(CommonConfig, LocalUserNetId->AsShared(), GetFirstFrameworkInstanceName(LocalUserNetId)))
+				if(LocalUserNetId.IsValid())
 				{
-					FAccountId LocalUserAccountId = LocalUserCommonAccountPtr->GetId(PresenceFrameworkInstance);
-
-					const FUniqueNetIdRepl& SocialUserId = SubsystemInfo->GetUserId();
-
-					if (const FCommonAccountPtr SocialUserCommonAccountPtr = GetCommonAccountFromV1(CommonConfig, SocialUserId->AsShared(), GetFirstFrameworkInstanceName(SocialUserId)))
+					if (const FCommonAccountPtr LocalUserCommonAccountPtr = GetCommonAccountFromV1(CommonConfig, LocalUserNetId->AsShared(), GetFirstFrameworkInstanceName(LocalUserNetId)))
 					{
-						FAccountId SocialUserAccountId = SocialUserCommonAccountPtr->GetId(PresenceFrameworkInstance);
+						FAccountId LocalUserAccountId = LocalUserCommonAccountPtr->GetId(PresenceFrameworkInstance);
 
-						if (IOnlineServicesPtr OnlineServices = CommonConfig.GetServices(PresenceFrameworkInstance))
+						const FUniqueNetIdRepl& SocialUserId = SubsystemInfo->GetUserId();
+						if (SocialUserId.IsValid())
 						{
-							if (IUserInfoPtr UserInfo = OnlineServices->GetUserInfoInterface())
+							if (const FCommonAccountPtr SocialUserCommonAccountPtr = GetCommonAccountFromV1(CommonConfig, SocialUserId->AsShared(), GetFirstFrameworkInstanceName(SocialUserId)))
 							{
-								TOnlineResult<FGetUserInfo> GetUserInfoResult = UserInfo->GetUserInfo({ LocalUserAccountId, SocialUserAccountId });
-								if (GetUserInfoResult.IsOk())
+								FAccountId SocialUserAccountId = SocialUserCommonAccountPtr->GetId(PresenceFrameworkInstance);
+
+								if (IOnlineServicesPtr OnlineServices = CommonConfig.GetServices(PresenceFrameworkInstance))
 								{
-									FString DisplayNameV2 = GetUserInfoResult.GetOkValue().UserInfo->DisplayName;
-#if !UE_BUILD_SHIPPING
-									// Temporary redundancy check
-									if (DisplayName != DisplayNameV2)
+									if (IUserInfoPtr UserInfo = OnlineServices->GetUserInfoInterface())
 									{
-										UE_LOG(LogParty, Warning, TEXT("DisplayName disparity: DisplayNameV1: %s. DisplayNameV2: %s"), *DisplayName, *DisplayNameV2);
-									}
+										TOnlineResult<FGetUserInfo> GetUserInfoResult = UserInfo->GetUserInfo({ LocalUserAccountId, SocialUserAccountId });
+										if (GetUserInfoResult.IsOk())
+										{
+											FString DisplayNameV2 = GetUserInfoResult.GetOkValue().UserInfo->DisplayName;
+#if !UE_BUILD_SHIPPING
+											// Temporary redundancy check
+											if (DisplayName != DisplayNameV2)
+											{
+												UE_LOG(LogParty, Warning, TEXT("DisplayName disparity: DisplayNameV1: %s. DisplayNameV2: %s"), *DisplayName, *DisplayNameV2);
+											}
 #endif
-									if (!DisplayNameV2.IsEmpty())
-									{
-										DisplayName = DisplayNameV2;
+											if (!DisplayNameV2.IsEmpty())
+											{
+												DisplayName = DisplayNameV2;
+											}
+										}
 									}
 								}
 							}
@@ -1034,18 +1046,24 @@ UE::Online::TOnlineResult<UE::Online::FGetCachedPresence> USocialUser::GetFriend
 	const FName& PresenceFrameworkInstance = USocialManager::GetPresenceFrameworkInstance();
 	FCommonConfig CommonConfig(this);
 	FUniqueNetIdRepl LocalUserNetId = GetOwningToolkit().GetLocalUserNetId(SubsystemType);
-	if (const FCommonAccountPtr LocalUserCommonAccountPtr = GetCommonAccountFromV1(CommonConfig, LocalUserNetId->AsShared(), GetFirstFrameworkInstanceName(LocalUserNetId)))
+	if (LocalUserNetId.IsValid())
 	{
-		LocalUserAccountId = LocalUserCommonAccountPtr->GetId(PresenceFrameworkInstance);
+		if (const FCommonAccountPtr LocalUserCommonAccountPtr = GetCommonAccountFromV1(CommonConfig, LocalUserNetId->AsShared(), GetFirstFrameworkInstanceName(LocalUserNetId)))
+		{
+			LocalUserAccountId = LocalUserCommonAccountPtr->GetId(PresenceFrameworkInstance);
+		}
 	}
 
+
 	if(const FSubsystemUserInfo* SubsystemInfo = SubsystemInfoByType.Find(SubsystemType))
 	{
 		FUniqueNetIdRepl FriendUserId = SubsystemInfo->GetUserId();
-
-		if (const FCommonAccountPtr FriendCommonAccountPtr = GetCommonAccountFromV1(CommonConfig, FriendUserId->AsShared(), GetFirstFrameworkInstanceName(FriendUserId)))
+		if (FriendUserId.IsValid())
 		{
-			FriendAccountId = FriendCommonAccountPtr->GetId(PresenceFrameworkInstance);
+			if (const FCommonAccountPtr FriendCommonAccountPtr = GetCommonAccountFromV1(CommonConfig, FriendUserId->AsShared(), GetFirstFrameworkInstanceName(FriendUserId)))
+			{
+				FriendAccountId = FriendCommonAccountPtr->GetId(PresenceFrameworkInstance);
+			}
 		}
 	}
 
