Make sure to pass correct friend unique net id when find session in FPartyPlatformSessionMonitor;
Implement FindSessionById in OnlineSessionInterfaceMpaGDK.

#Synced-CL 48344811

diff --git a/Engine/Plugins/Online/OnlineFramework/Source/Party/Private/Party/PartyPlatformSessionMonitor.cpp b/Engine/Plugins/Online/OnlineFramework/Source/Party/Private/Party/PartyPlatformSessionMonitor.cpp
--- a/Engine/Plugins/Online/OnlineFramework/Source/Party/Private/Party/PartyPlatformSessionMonitor.cpp
+++ b/Engine/Plugins/Online/OnlineFramework/Source/Party/Private/Party/PartyPlatformSessionMonitor.cpp
@@ -235,7 +235,7 @@ bool FPartyPlatformSessionManager::FindSessionInternal(const FSessionId& Session
 								}
 								else
 								{
-									GetSessionInterface()->FindSessionById(*LocalUserPlatformId, *SessionId, *LocalUserPlatformId, FOnSingleSessionResultCompleteDelegate::CreateSP(this, &FPartyPlatformSessionManager::HandleFindSessionByIdComplete, SessionIdString, SessionOwnerId, OnAttemptComplete));
+									GetSessionInterface()->FindSessionById(*LocalUserPlatformId, *SessionId, *SessionOwnerId.GetV1(), FOnSingleSessionResultCompleteDelegate::CreateSP(this, &FPartyPlatformSessionManager::HandleFindSessionByIdComplete, SessionIdString, SessionOwnerId, OnAttemptComplete));
 								}
 							}
 							return false; // Don't retick
@@ -244,7 +244,7 @@ bool FPartyPlatformSessionManager::FindSessionInternal(const FSessionId& Session
 				}
 #endif
 				// Always start by trying to find the session directly by ID
-				return SessionInterface->FindSessionById(*LocalUserPlatformId, *SessionId, *LocalUserPlatformId, FOnSingleSessionResultCompleteDelegate::CreateSP(this, &FPartyPlatformSessionManager::HandleFindSessionByIdComplete, SessionIdString, SessionOwnerId, OnAttemptComplete));
+				return SessionInterface->FindSessionById(*LocalUserPlatformId, *SessionId, *SessionOwnerId.GetV1(), FOnSingleSessionResultCompleteDelegate::CreateSP(this, &FPartyPlatformSessionManager::HandleFindSessionByIdComplete, SessionIdString, SessionOwnerId, OnAttemptComplete));
 			}
 			else
 			{
