Set log verbosity for important LogParty lines to Log

#Synced-CL 48461929

diff --git a/Engine/Plugins/Online/OnlineFramework/Source/Party/Private/Party/PartyMember.cpp b/Engine/Plugins/Online/OnlineFramework/Source/Party/Private/Party/PartyMember.cpp
--- a/Engine/Plugins/Online/OnlineFramework/Source/Party/Private/Party/PartyMember.cpp
+++ b/Engine/Plugins/Online/OnlineFramework/Source/Party/Private/Party/PartyMember.cpp
@@ -351,19 +351,19 @@ void UPartyMember::NotifyMemberDataReceived(const FOnlinePartyData& MemberData)
 
 void UPartyMember::NotifyMemberPromoted()
 {
-	UE_LOG(LogParty, Verbose, TEXT("Member [%s] promoted to party leader."), *ToDebugString());
+	UE_LOG(LogParty, Log, TEXT("Member [%s] promoted to party leader."), *ToDebugString());
 	OnMemberPromotedInternal();
 }
 
 void UPartyMember::NotifyMemberDemoted()
 {
-	UE_LOG(LogParty, Verbose, TEXT("Member [%s] is no longer party leader."), *ToDebugString());
+	UE_LOG(LogParty, Log, TEXT("Member [%s] is no longer party leader."), *ToDebugString());
 	OnMemberDemotedInternal();
 }
 
 void UPartyMember::NotifyRemovedFromParty(EMemberExitedReason ExitReason)
 {
-	UE_LOG(LogParty, Verbose, TEXT("Member [%s] is no longer in the party. Reason = [%s]"), *ToDebugString(), ToString(ExitReason));
+	UE_LOG(LogParty, Log, TEXT("Member [%s] is no longer in the party. Reason = [%s]"), *ToDebugString(), ToString(ExitReason));
 	OnRemovedFromPartyInternal(ExitReason);
 }
 
diff --git a/Engine/Plugins/Online/OnlineFramework/Source/Party/Private/Party/PartyPlatformSessionMonitor.cpp b/Engine/Plugins/Online/OnlineFramework/Source/Party/Private/Party/PartyPlatformSessionMonitor.cpp
--- a/Engine/Plugins/Online/OnlineFramework/Source/Party/Private/Party/PartyPlatformSessionMonitor.cpp
+++ b/Engine/Plugins/Online/OnlineFramework/Source/Party/Private/Party/PartyPlatformSessionMonitor.cpp
@@ -157,7 +157,7 @@ void FPartyPlatformSessionManager::CreateMonitor(USocialParty& Party)
 
 bool FPartyPlatformSessionManager::FindSession(const USocialUser& User, const FOnFindSessionAttemptComplete& OnAttemptComplete)
 {
-	UE_LOG(LogParty, Verbose, TEXT("FPartyPlatformSessionManager finding party platform session of user [%s]"), *User.GetDisplayName());
+	UE_LOG(LogParty, Log, TEXT("FPartyPlatformSessionManager finding party platform session of user [%s]"), *User.GetDisplayName());
 
 	FSessionId SessionId;
 	FUniqueNetIdRepl UserPlatformId;
@@ -175,7 +175,7 @@ bool FPartyPlatformSessionManager::FindSession(const USocialUser& User, const FO
 
 bool FPartyPlatformSessionManager::FindSession(const FPartyPlatformSessionInfo& SessionInfo, const FOnFindSessionAttemptComplete& OnAttemptComplete)
 {
-	UE_LOG(LogParty, Verbose, TEXT("FPartyPlatformSessionManager finding party platform session [%s]"), *SessionInfo.ToDebugString());
+	UE_LOG(LogParty, Log, TEXT("FPartyPlatformSessionManager finding party platform session [%s]"), *SessionInfo.ToDebugString());
 
 	const USocialToolkit* SocialToolkit = SocialManager.GetFirstLocalUserToolkit();
 	check(SocialToolkit);
@@ -257,7 +257,7 @@ bool FPartyPlatformSessionManager::FindSessionInternal(const FSessionId& Session
 
 void FPartyPlatformSessionManager::HandleFindSessionByIdComplete(int32 LocalUserNum, bool bWasSuccessful, const FOnlineSessionSearchResult& FoundSession, FSessionId SessionId, FUniqueNetIdRepl SessionOwnerId, FOnFindSessionAttemptComplete OnAttemptComplete)
 {
-	UE_LOG(LogParty, Verbose, TEXT("FPartyPlatformSessionManager completed attempt to find platform session [%s] of user [%s] by SessionId with result [%d]"), *SessionId, *SessionOwnerId.ToDebugString(), (int32)bWasSuccessful);
+	UE_LOG(LogParty, Log, TEXT("FPartyPlatformSessionManager completed attempt to find platform session [%s] of user [%s] by SessionId with result [%d]"), *SessionId, *SessionOwnerId.ToDebugString(), (int32)bWasSuccessful);
 
 	if (bWasSuccessful && FoundSession.IsSessionInfoValid())
 	{
@@ -332,7 +332,7 @@ void FPartyPlatformSessionManager::HandleReadPSNFriendsListComplete(int32 LocalU
 
 void FPartyPlatformSessionManager::HandleFindFriendSessionsComplete(int32 LocalUserNum, bool bWasSuccessful, const TArray<FOnlineSessionSearchResult>& FoundSessions, FSessionId SessionId, FUniqueNetIdRepl SessionOwnerId, FOnFindSessionAttemptComplete OnAttemptComplete)
 {
-	UE_LOG(LogParty, Verbose, TEXT("PartyPlatformSessionManager found [%d] sessions searching by friend ID [%s] (bWasSuccessful=%d)"), FoundSessions.Num(), *SessionOwnerId.ToDebugString(), (int32)bWasSuccessful);
+	UE_LOG(LogParty, Log, TEXT("PartyPlatformSessionManager found [%d] sessions searching by friend ID [%s] (bWasSuccessful=%d)"), FoundSessions.Num(), *SessionOwnerId.ToDebugString(), (int32)bWasSuccessful);
 
 	GetSessionInterface()->ClearOnFindFriendSessionCompleteDelegates(LocalUserNum, this);
 
@@ -511,7 +511,7 @@ void FPartyPlatformSessionMonitor::Initialize()
 	check(MonitoredParty.IsValid());
 
 	USocialParty& Party = *MonitoredParty;
-	UE_LOG(LogParty, Verbose, TEXT("Initializing PartyPlatformSessionMonitor for party [%s]"), *Party.GetPartyId().ToDebugString());
+	UE_LOG(LogParty, Log, TEXT("Initializing PartyPlatformSessionMonitor for party [%s]"), *Party.GetPartyId().ToDebugString());
 
 	Party.OnPartyMemberCreated().AddSP(this, &FPartyPlatformSessionMonitor::HandlePartyMemberCreated);
 	Party.OnPartyConfigurationChanged().AddSP(this, &FPartyPlatformSessionMonitor::HandlePartyConfigurationChanged);
@@ -531,7 +531,7 @@ void FPartyPlatformSessionMonitor::Initialize()
 
 void FPartyPlatformSessionMonitor::ShutdownInternal()
 {
-	UE_LOG(LogParty, Verbose, TEXT("Finalizing shutdown of PartyPlatformSessionMonitor for party of type [%d]"), PartyTypeId.GetValue());
+	UE_LOG(LogParty, Log, TEXT("Finalizing shutdown of PartyPlatformSessionMonitor for party of type [%d]"), PartyTypeId.GetValue());
 
 	if (RetryTickerHandle.IsValid())
 	{
@@ -587,7 +587,7 @@ void FPartyPlatformSessionMonitor::CreateSession(const FUniqueNetIdRepl& LocalUs
 		}
 #endif
 		SessionInterface->CreateSession(*LocalUserPlatformId, NAME_PartySession, SessionSettings);
-		UE_LOG(LogParty, Verbose, TEXT("PartyPlatformSessionMonitor creating session with the following parameters: "));
+		UE_LOG(LogParty, Log, TEXT("PartyPlatformSessionMonitor creating session with the following parameters: "));
 		DumpSessionSettings(&SessionSettings);
 	}
 }
@@ -800,7 +800,7 @@ void FPartyPlatformSessionMonitor::HandlePlatformSessionsChanged()
 
 	if (!bWasSessionOwner && DoesLocalUserOwnPlatformSession())
 	{
-		UE_LOG(LogParty, Verbose, TEXT("Local user just became owner of their party platform session within party [%s]"), *MonitoredParty->ToDebugString());
+		UE_LOG(LogParty, Log, TEXT("Local user just became owner of their party platform session within party [%s]"), *MonitoredParty->ToDebugString());
 
 		// We just took over ownership of the session on this platform
 		QueuePlatformSessionUpdate();
@@ -1112,7 +1112,7 @@ void FPartyPlatformSessionMonitor::HandleJoinSessionComplete(FName SessionName,
 
 void FPartyPlatformSessionMonitor::HandleDestroySessionComplete(FName SessionName, bool bWasSuccessful)
 {
-	UE_LOG(LogParty, Verbose, TEXT("PartyPlatformSessionMonitor finished destroying party session with result [%d]."), (int32)bWasSuccessful);
+	UE_LOG(LogParty, Log, TEXT("PartyPlatformSessionMonitor finished destroying party session with result [%d]."), (int32)bWasSuccessful);
 
 	//@todo DanH Sessions: What does it mean to fail at leaving a session? Does that mean we need to try again? Or just that we weren't in one to begin with? #suggested
 	ShutdownInternal();
@@ -1253,7 +1253,7 @@ void FPartyPlatformSessionMonitor::ProcessJoinFailure()
 	const FPartyPlatformSessionInfo* ExistingSessionInfo = FindLocalPlatformSessionInfo();
 	if (ExistingSessionInfo && ExistingSessionInfo->SessionId != TargetSessionId)
 	{
-		UE_LOG(LogParty, Verbose, TEXT("PartyPlatformSessionMonitor targeted platform session [%s] out of date - retrying with updated id [%s]"), *TargetSessionId, *ExistingSessionInfo->SessionId);
+		UE_LOG(LogParty, Log, TEXT("PartyPlatformSessionMonitor targeted platform session [%s] out of date - retrying with updated id [%s]"), *TargetSessionId, *ExistingSessionInfo->SessionId);
 		TargetSessionId = FSessionId();
 
 		// The ID of our platform session changed during the find attempt, so try again right away
diff --git a/Engine/Plugins/Online/OnlineFramework/Source/Party/Private/Party/SocialParty.cpp b/Engine/Plugins/Online/OnlineFramework/Source/Party/Private/Party/SocialParty.cpp
--- a/Engine/Plugins/Online/OnlineFramework/Source/Party/Private/Party/SocialParty.cpp
+++ b/Engine/Plugins/Online/OnlineFramework/Source/Party/Private/Party/SocialParty.cpp
@@ -1338,7 +1338,7 @@ void USocialParty::UpdatePartyConfig(bool bResetAccessKey)
 {
 	check(IsLocalPlayerPartyLeader());
 
-	UE_LOG(LogParty, Verbose, TEXT("Party [%s] attempting to update party config"), *ToDebugString());
+	UE_LOG(LogParty, Log, TEXT("Party [%s] attempting to update party config"), *ToDebugString());
 
 	IOnlinePartyPtr PartyInterface = Online::GetPartyInterfaceChecked(GetWorld());
 	PartyInterface->UpdateParty(*OwningLocalUserId, GetPartyId(), CurrentConfig, bResetAccessKey, FOnUpdatePartyComplete::CreateUObject(this, &USocialParty::HandleUpdatePartyConfigComplete));
@@ -1359,7 +1359,7 @@ void USocialParty::LeaveParty(const FOnLeavePartyAttemptComplete& OnLeaveAttempt
 	}
 	else
 	{
-		UE_LOG(LogParty, Verbose, TEXT("Attempting to leave party [%s]"), *ToDebugString());
+		UE_LOG(LogParty, Log, TEXT("Attempting to leave party [%s]"), *ToDebugString());
 
 		BeginLeavingParty(EMemberExitedReason::Left);
 
@@ -1387,7 +1387,7 @@ void USocialParty::RemoveLocalMember(const FUniqueNetIdRepl& LocalUserId, const
 	}
 	else
 	{
-		UE_LOG(LogParty, Verbose, TEXT("Attempting to leave party LocalUserId=[%s] [%s] "), *LocalUserId.ToDebugString(), *ToDebugString());
+		UE_LOG(LogParty, Log, TEXT("Attempting to leave party LocalUserId=[%s] [%s] "), *LocalUserId.ToDebugString(), *ToDebugString());
 
 		const FOnlinePartyId& PartyId = GetPartyId();
 		if (LocalUserId.IsValid() && PartyId.IsValid())
@@ -1838,7 +1838,7 @@ void USocialParty::DisconnectParty()
 void USocialParty::FinalizePartyLeave(EMemberExitedReason Reason)
 {
 	QUICK_SCOPE_CYCLE_COUNTER(STAT_USocialParty_FinalizePartyLeave);
-	UE_LOG(LogParty, Verbose, TEXT("Local player [%s] is no longer in party [%s]. Reason [%s]."), *GetOwningLocalMember().ToDebugString(false), *ToDebugString(), ToString(Reason));
+	UE_LOG(LogParty, Log, TEXT("Local player [%s] is no longer in party [%s]. Reason [%s]."), *GetOwningLocalMember().ToDebugString(false), *ToDebugString(), ToString(Reason));
 
 	if (!bIsLeavingParty)
 	{
@@ -1956,7 +1956,7 @@ void USocialParty::RequestJoinInProgress(const UPartyMember& TargetMember, const
 	Request.Target = TargetMember.GetPrimaryNetId();
 	Request.Time = FDateTime::UtcNow().ToUnixTimestamp();
 
-	UE_LOG(LogParty, Verbose, TEXT("RequestJoinInProgress: Sending request Target=%s Time=%d"), *Request.Target.ToDebugString(), Request.Time);
+	UE_LOG(LogParty, Log, TEXT("RequestJoinInProgress: Sending request Target=%s Time=%d"), *Request.Target.ToDebugString(), Request.Time);
 	for (UPartyMember* RequestingMember : GetLocalPartyMembersForJoinInProgress())
 	{
 		if (RequestingMember && ensure(RequestingMember->IsLocalPlayer()))
@@ -2012,7 +2012,7 @@ void USocialParty::HandleJoinInProgressDataRequestChanged(const FPartyMemberJoin
 		return;
 	}
 
-	UE_LOG(LogParty, Verbose, TEXT("HandleJoinInProgressDataRequestChanged: Received request Requester=%s Time=%d"), *Member->GetPrimaryNetId().ToDebugString(), Request.Time);
+	UE_LOG(LogParty, Log, TEXT("HandleJoinInProgressDataRequestChanged: Received request Requester=%s Time=%d"), *Member->GetPrimaryNetId().ToDebugString(), Request.Time);
 	FPendingMemberApproval PendingApproval;
 	PendingApproval.RecipientId = OwningLocalUserId;
 	PendingApproval.Members.Emplace(Member->GetPrimaryNetId(), Member->GetRepData().GetPlatformDataPlatform());
@@ -2045,7 +2045,7 @@ void USocialParty::RespondToJoinInProgressRequest(const FPendingMemberApproval&
 	Response.ResponseTime = FDateTime::UtcNow().ToUnixTimestamp();
 	Response.DenialReason = static_cast<uint8>(DenialReason);
 
-	UE_LOG(LogParty, Verbose, TEXT("RespondToJoinInProgressRequest: Sending response Requester=%s RequestTime=%d ResponseTime=%d DenialReason=%s"),
+	UE_LOG(LogParty, Log, TEXT("RespondToJoinInProgressRequest: Sending response Requester=%s RequestTime=%d ResponseTime=%d DenialReason=%s"),
 		*Response.Requester.ToDebugString(), Response.RequestTime, Response.ResponseTime, ToString(DenialReason));
 	TArray<FPartyMemberJoinInProgressResponse> Responses = GetOwningLocalMember().GetRepData().GetJoinInProgressDataResponses();
 	Responses.Add(Response);
@@ -2089,7 +2089,7 @@ void USocialParty::HandleJoinInProgressDataResponsesChanged(const TArray<FPartyM
 
 		// Responses are ordered newest first, so use the first one we find.
 		const int64 Now = FDateTime::UtcNow().ToUnixTimestamp();
-		UE_LOG(LogParty, Verbose, TEXT("HandleJoinInProgressDataResponsesChanged: Received response Requester=%s RequestTime=%d ResponseTime=%d DenialReason=%s RTT=%d"),
+		UE_LOG(LogParty, Log, TEXT("HandleJoinInProgressDataResponsesChanged: Received response Requester=%s RequestTime=%d ResponseTime=%d DenialReason=%s RTT=%d"),
 			*Response.Requester.ToDebugString(), Response.RequestTime, Response.ResponseTime, ToString(DenialReason), Now - Response.RequestTime);
 		CallJoinInProgressComplete(DenialReason);
 		break;
diff --git a/Engine/Plugins/Online/OnlineFramework/Source/Party/Private/SocialManager.cpp b/Engine/Plugins/Online/OnlineFramework/Source/Party/Private/SocialManager.cpp
--- a/Engine/Plugins/Online/OnlineFramework/Source/Party/Private/SocialManager.cpp
+++ b/Engine/Plugins/Online/OnlineFramework/Source/Party/Private/SocialManager.cpp
@@ -601,7 +601,7 @@ bool USocialManager::CanAcceptInvitationsFromAwayUsers() const
 
 void USocialManager::JoinParty(const USocialUser& UserToJoin, const FOnlinePartyTypeId& PartyTypeId, const FOnJoinPartyAttemptComplete& OnJoinPartyComplete, const FName& JoinMethod)
 {
-	UE_LOG(LogParty, Verbose, TEXT("Attempting to join user [%s]'s party of type [%d] by [%s]"), *UserToJoin.ToDebugString(), PartyTypeId.GetValue(), *JoinMethod.ToString());
+	UE_LOG(LogParty, Log, TEXT("Attempting to join user [%s]'s party of type [%d] by [%s]"), *UserToJoin.ToDebugString(), PartyTypeId.GetValue(), *JoinMethod.ToString());
 
 	FJoinPartyAttempt NewAttempt(&UserToJoin, PartyTypeId, JoinMethod, OnJoinPartyComplete);
 	NewAttempt.AnalyticsContext = UserToJoin.GetAnalyticsContext();
@@ -679,7 +679,7 @@ TSubclassOf<USocialParty> USocialManager::GetPartyClassForType(const FOnlinePart
 
 void USocialManager::OnJoinPartyAttemptCompleteInternal(const FJoinPartyAttempt& JoinAttemptInfo, const FJoinPartyResult& Result)
 {
-	UE_LOG(LogParty, Verbose, TEXT("JoinPartyAttempt [%s] completed with result [%s] and reason [%s]"), *JoinAttemptInfo.ToDebugString(), ToString(Result.GetResult()), ToString(Result.GetDenialReason().GetReason()));
+	UE_LOG(LogParty, Log, TEXT("JoinPartyAttempt [%s] completed with result [%s] and reason [%s]"), *JoinAttemptInfo.ToDebugString(), ToString(Result.GetResult()), ToString(Result.GetDenialReason().GetReason()));
 }
 
 void USocialManager::OnToolkitCreatedInternal(USocialToolkit& NewToolkit)
@@ -1041,7 +1041,7 @@ void USocialManager::HandleToolkitReset(int32 LocalUserNum)
 
 void USocialManager::RestorePartyStateFromPartySystem(const FOnRestorePartyStateFromPartySystemComplete& OnRestoreComplete)
 {
-	UE_LOG(LogParty, Verbose, TEXT("RestorePartyStateFromPartySystem"));
+	UE_LOG(LogParty, Log, TEXT("RestorePartyStateFromPartySystem"));
 	FUniqueNetIdRepl LocalUserId = GetFirstLocalUserId(ESocialSubsystem::Primary);
 
 	// If the player has any parties, do not try to restore
@@ -1136,7 +1136,6 @@ void USocialManager::HandleCreatePartyComplete(const FUniqueNetId& LocalUserId,
 
 			if (UPartyMember* LocalPartyMember = NewParty->GetPartyMember(LocalUserId.AsShared()))
 			{
-				UE_LOG(LogParty, Verbose, TEXT("Player [%s] created party."), *LocalUserId.ToDebugString());
 				LocalPartyMember->GetMutableRepData().SetJoinMethod(PartyJoinMethod::Creation.ToString());
 			}
 		}
@@ -1210,7 +1209,7 @@ void USocialManager::HandlePersistentPartyStateChanged(EPartyState NewState, EPa
 {
 	ABORT_DURING_SHUTDOWN();
 
-	UE_LOG(LogParty, Verbose, TEXT("Persistent party state changed to %s"), ToString(NewState));
+	UE_LOG(LogParty, Log, TEXT("Persistent party state changed to %s"), ToString(NewState));
 
 	if (NewState == EPartyState::Disconnected)
 	{
@@ -1236,7 +1235,7 @@ void USocialManager::HandleLeavePartyForJoinComplete(ELeavePartyCompletionResult
 
 	if (LeftParty)
 	{
-		UE_LOG(LogParty, Verbose, TEXT("Attempt to leave party [%s] for pending join completed with result [%s]"), *LeftParty->ToDebugString(), ToString(LeaveResult));
+		UE_LOG(LogParty, Log, TEXT("Attempt to leave party [%s] for pending join completed with result [%s]"), *LeftParty->ToDebugString(), ToString(LeaveResult));
 	}
 }
 
@@ -1306,7 +1305,7 @@ void USocialManager::HandlePartyLeft(EMemberExitedReason Reason, USocialParty* L
 			{
 				// An attempt to join a party of this type has been initiated, but something/someone decided to leave the party before the attempt was ready to do so
 				// It's not worth accounting for the potential limbo that this could put us into, so just abort the join attempt and let the explicit leave action win
-				UE_LOG(LogParty, Verbose, TEXT("Finished leaving party [%s] before the current join attempt established join info. Cancelling join attempt."), *LeftParty->ToDebugString());
+				UE_LOG(LogParty, Log, TEXT("Finished leaving party [%s] before the current join attempt established join info. Cancelling join attempt."), *LeftParty->ToDebugString());
 				FinishJoinPartyAttempt(*JoinAttempt, FJoinPartyResult(EPartyJoinDenialReason::JoinAttemptAborted));
 			}
 		}
