Bugfix for proper presence status retrieval

diff --git a/Engine/Plugins/Online/OnlineFramework/Source/Party/Private/User/SocialUser.cpp b/Engine/Plugins/Online/OnlineFramework/Source/Party/Private/User/SocialUser.cpp
--- a/Engine/Plugins/Online/OnlineFramework/Source/Party/Private/User/SocialUser.cpp
+++ b/Engine/Plugins/Online/OnlineFramework/Source/Party/Private/User/SocialUser.cpp
@@ -615,7 +615,7 @@ UE::Online::EUserPresenceStatus USocialUser::GetOnlineStatusV2() const
 				}
 				else if (SubsystemInfoPair.Key == ESocialSubsystem::Platform) // For Platform presence, we'll get it from V1 and convert it
 				{
-					EOnlinePresenceState::Type OnlineStatus = EOnlinePresenceState::Offline;
+					EOnlinePresenceState::Type OnlineStatus = UE::Online::FPresenceOSSAdapter::StatusV2toV1(Result);
 
 					if (const FOnlineUserPresence* PresenceInfo = SubsystemInfoPair.Value.GetPresenceInfo())
 					{
