Added delegate reacting to userinfo CVar change to trigger V2 query

#Synced-CL 48774492

diff --git a/Engine/Plugins/Online/OnlineFramework/Source/Party/Private/SocialToolkit.cpp b/Engine/Plugins/Online/OnlineFramework/Source/Party/Private/SocialToolkit.cpp
--- a/Engine/Plugins/Online/OnlineFramework/Source/Party/Private/SocialToolkit.cpp
+++ b/Engine/Plugins/Online/OnlineFramework/Source/Party/Private/SocialToolkit.cpp
@@ -723,7 +723,7 @@ void USocialToolkit::OnOwnerLoggedIn()
 
 				if (IConsoleVariable* CVar = IConsoleManager::Get().FindConsoleVariable(TEXT("onlineframework.party.shouldreadpresencefromonlineservices")))
 				{
-					CVar->OnChangedDelegate().AddUObject(this, &USocialToolkit::HandlePresenceReadCVarChanged);
+					OnReadPresenceV2CVarChanged = CVar->OnChangedDelegate().AddUObject(this, &USocialToolkit::HandlePresenceReadCVarChanged);
 				}
 			}
 		}
@@ -790,6 +790,16 @@ void USocialToolkit::OnOwnerLoggedOut()
 				PresenceInterface->ClearOnPresenceReceivedDelegates(this);
 			}
 
+			if (OnReadPresenceV2CVarChanged.IsValid())
+			{
+				if (IConsoleVariable* CVar = IConsoleManager::Get().FindConsoleVariable(TEXT("onlineframework.party.shouldreadpresencefromonlineservices")))
+				{
+					CVar->OnChangedDelegate().Remove(OnReadPresenceV2CVarChanged);
+				}
+
+				OnReadPresenceV2CVarChanged.Reset();
+			}
+
 			if (UE::Online::FOnlineEventDelegateHandle* Handle = OnPresenceUpdatedHandles.Find(SubsystemType))
 			{
 				Handle->Unbind();
diff --git a/Engine/Plugins/Online/OnlineFramework/Source/Party/Private/User/SocialUser.cpp b/Engine/Plugins/Online/OnlineFramework/Source/Party/Private/User/SocialUser.cpp
--- a/Engine/Plugins/Online/OnlineFramework/Source/Party/Private/User/SocialUser.cpp
+++ b/Engine/Plugins/Online/OnlineFramework/Source/Party/Private/User/SocialUser.cpp
@@ -155,7 +155,7 @@ public:
 						{
 							FAccountId LocalUserAccountId = LocalUserCommonAccountPtr->GetId(USocialManager::GetPresenceFrameworkInstance());
 
-							UE_LOG(LogParty, Log, TEXT("FSocialQuery_UserInfoV2 executing for [%d] users on subsystem [%s]"), SocialUserAccountIds.Num(), ToString(SubsystemType));
+							UE_LOG(LogParty, Log, TEXT("FSocialQuery_UserInfoV2 executing for [%d] users on services [%s]"), SocialUserAccountIds.Num(), ToString(SubsystemType));
 
 							UserInfo->QueryUserInfo({ LocalUserAccountId, SocialUserAccountIds })
 								.OnComplete([this, LocalUserAccountId, SocialUserAccountIds, UserInfoWeakPtr = UserInfo.ToWeakPtr()](const TOnlineResult<FQueryUserInfo>& Result)
@@ -330,6 +330,19 @@ const FOnlineUserPresence* USocialUser::FSubsystemUserInfo::GetPresenceInfo() co
 USocialUser::USocialUser()
 {}
 
+USocialUser::~USocialUser()
+{
+	if (OnReadUserInfoV2CVarChanged.IsValid())
+	{
+		if (IConsoleVariable* CVar = IConsoleManager::Get().FindConsoleVariable(TEXT("onlineframework.party.shouldreaduserinfofromonlineservices")))
+		{
+			CVar->OnChangedDelegate().Remove(OnReadUserInfoV2CVarChanged);
+		}
+
+		OnReadUserInfoV2CVarChanged.Reset();
+	}
+}
+
 void USocialUser::InitLocalUser()
 {
 	check(IsLocalUser());
@@ -2083,7 +2096,8 @@ void USocialUser::SetSubsystemId(ESocialSubsystem SubsystemType, const FUniqueNe
 
 		USocialToolkit& OwningToolkit = GetOwningToolkit();
 		OwningToolkit.NotifySubsystemIdEstablished(*this, SubsystemType, SubsystemId);
-
+		
+		// OSS UserInfo
 		IOnlineSubsystem* OSS = OwningToolkit.GetSocialOss(SubsystemType);
 		if (ensure(OSS))
 		{
@@ -2100,13 +2114,69 @@ void USocialUser::SetSubsystemId(ESocialSubsystem SubsystemType, const FUniqueNe
 				{
 					UE_LOG(LogParty, VeryVerbose, TEXT("SocialUser [%s] querying user info on subsystem [%s]"), *ToDebugString(), ToString(SubsystemType));
 
-					// No valid user info exists on this subsystem, so queue up a query for it
+					// No valid user info exists on this subsystem, so queue up a query for it in both subsystem and services, if enabled
+
 					auto QueryCompleteHandler = FSocialQuery_UserInfo::FOnQueryComplete::CreateUObject(this, &USocialUser::HandleQueryUserInfoComplete);
 					FSocialQueryManager::AddUserId<FSocialQuery_UserInfo>(OwningToolkit, SubsystemType, SubsystemId.GetUniqueNetId().ToSharedRef(), QueryCompleteHandler);
 					NumPendingQueries++;
 				}
 			}
 		}
+
+		// OnlineServices UserInfo
+		if (SubsystemType == ESocialSubsystem::Primary)
+		{
+			auto QueryUserInfoV2 = [WeakThis = TWeakObjectPtr<USocialUser>(this)]()
+			{
+				if (USocialUser* StrongThis = WeakThis.Get())
+				{
+					if(StrongThis->ShouldReadUserInfoFromOnlineServices())
+					{
+						using namespace UE::OnlineFramework;
+
+						UE_LOG(LogParty, VeryVerbose, TEXT("SocialUser [%s] querying user info on services [%s]"), *StrongThis->ToDebugString(), ToString(ESocialSubsystem::Primary));
+
+						const FUniqueNetIdRepl UserId = StrongThis->GetUserId(ESocialSubsystem::Primary);
+						if (UserId.IsValid())
+						{
+							FCommonConfig CommonConfig(StrongThis);
+							if (const FCommonAccountPtr CommonAccountPtr = GetCommonAccountFromV1(CommonConfig, UserId->AsShared(), GetFirstFrameworkInstanceName(UserId)))
+							{
+								UE::Online::FAccountId AccountId = CommonAccountPtr->GetId(USocialManager::GetPresenceFrameworkInstance());
+								if (AccountId.IsValid())
+								{
+									// No valid user info exists on this subsystem, so queue up a query for it
+									FSocialQuery_UserInfoV2::FOnQueryComplete QueryCompleteHandler = FSocialQuery_UserInfoV2::FOnQueryComplete::CreateUObject(StrongThis, &USocialUser::HandleQueryUserInfoV2Complete);
+									FSocialQueryManager::AddUserId<FSocialQuery_UserInfoV2>(StrongThis->GetOwningToolkit(), ESocialSubsystem::Primary, AccountId, QueryCompleteHandler);
+												
+									if (!StrongThis->bIsInitialized)
+									{
+										StrongThis->NumPendingQueries++;
+									}
+								}
+							}
+						}
+					}
+				}
+			};
+
+			// We want to re-query user info whenever the functionality is enabled
+			if (!OnReadUserInfoV2CVarChanged.IsValid())
+			{
+				if (IConsoleVariable* CVar = IConsoleManager::Get().FindConsoleVariable(TEXT("onlineframework.party.shouldreaduserinfofromonlineservices")))
+				{
+					OnReadUserInfoV2CVarChanged = CVar->OnChangedDelegate().AddWeakLambda(this, [QueryUserInfoV2](IConsoleVariable*)
+						{
+							QueryUserInfoV2();
+						});
+				}
+			}			
+
+			if(ShouldReadUserInfoFromOnlineServices())
+			{
+				QueryUserInfoV2();
+			}
+		}
 	}
 }
 
@@ -2178,48 +2248,15 @@ void USocialUser::HandleQueryUserInfoComplete(ESocialSubsystem SubsystemType, bo
 		SetUserInfo(SubsystemType, UserInfo.ToSharedRef());
 	}
 
-	SetupCommonAccount()
-	.Then([WeakThis = TWeakObjectPtr<USocialUser>(this), SubsystemType, bIsUserInfoValid = UserInfo.IsValid()](TFuture<void> Result)
-	{
-		if (USocialUser* StrongThis = WeakThis.Get())
-		{
-			if (StrongThis->ShouldReadUserInfoFromOnlineServices() && SubsystemType == ESocialSubsystem::Primary)
-			{
-				using namespace UE::OnlineFramework;
-
-				UE_LOG(LogParty, VeryVerbose, TEXT("SocialUser [%s] querying user info on services [%s]"), *StrongThis->ToDebugString(), ToString(SubsystemType));
-
-				const FUniqueNetIdRepl UserId = StrongThis->GetUserId(SubsystemType);
-				if (UserId.IsValid())
-				{
-					FCommonConfig CommonConfig(StrongThis);
-					if (const FCommonAccountPtr CommonAccountPtr = GetCommonAccountFromV1(CommonConfig, UserId->AsShared(), GetFirstFrameworkInstanceName(UserId)))
-					{
-						UE::Online::FAccountId AccountId = CommonAccountPtr->GetId(USocialManager::GetPresenceFrameworkInstance());
-						if (AccountId.IsValid())
-						{
-							// No valid user info exists on this subsystem, so queue up a query for it
-							FSocialQuery_UserInfoV2::FOnQueryComplete QueryCompleteHandler = FSocialQuery_UserInfoV2::FOnQueryComplete::CreateUObject(StrongThis, &USocialUser::HandleQueryUserInfoV2Complete);
-							FSocialQueryManager::AddUserId<FSocialQuery_UserInfoV2>(StrongThis->GetOwningToolkit(), SubsystemType, AccountId, QueryCompleteHandler);
-							StrongThis->NumPendingQueries++;
-
-							return;
-						}
-					}
-				}
-			}			
-
-			UE_LOG(LogParty, VeryVerbose, TEXT("User [%s] finished querying user info on subsystem [%s] with result [%d]. [%d] queries still pending."), *StrongThis->ToDebugString(), ToString(SubsystemType), bIsUserInfoValid, StrongThis->NumPendingQueries);
-			StrongThis->TryBroadcastInitializationComplete();
-		}
-	});
+	UE_LOG(LogParty, VeryVerbose, TEXT("User [%s] finished querying user info on subsystem [%s] with result [%d]. [%d] queries still pending."), *ToDebugString(), ToString(SubsystemType), bWasSuccessful, NumPendingQueries);
+	TryBroadcastInitializationComplete();
 }
 
 void USocialUser::HandleQueryUserInfoV2Complete(ESocialSubsystem SubsystemType, bool bWasSuccessful, const TValueOrError<TSharedRef<UE::Online::FUserInfo>, UE::Online::FOnlineError>& Result)
 {
 	--NumPendingQueries;
 
-	UE_LOG(LogParty, VeryVerbose, TEXT("User [%s] finished querying user info on subsystem [%s] with result [%d]. [%d] queries still pending."), *ToDebugString(), ToString(SubsystemType), bWasSuccessful, NumPendingQueries);
+	UE_LOG(LogParty, VeryVerbose, TEXT("User [%s] finished querying user info on services [%s] with result [%d]. [%d] queries still pending."), *ToDebugString(), ToString(SubsystemType), bWasSuccessful, NumPendingQueries);
 	TryBroadcastInitializationComplete();
 }
 
diff --git a/Engine/Plugins/Online/OnlineFramework/Source/Party/Public/SocialToolkit.h b/Engine/Plugins/Online/OnlineFramework/Source/Party/Public/SocialToolkit.h
--- a/Engine/Plugins/Online/OnlineFramework/Source/Party/Public/SocialToolkit.h
+++ b/Engine/Plugins/Online/OnlineFramework/Source/Party/Public/SocialToolkit.h
@@ -314,6 +314,7 @@ private:	// Handlers
 #endif
 
 	TMap<ESocialSubsystem, UE::Online::FOnlineEventDelegateHandle> OnPresenceUpdatedHandles;
+	FDelegateHandle OnReadPresenceV2CVarChanged;
 
 protected:
 	UPROPERTY()
diff --git a/Engine/Plugins/Online/OnlineFramework/Source/Party/Public/User/SocialUser.h b/Engine/Plugins/Online/OnlineFramework/Source/Party/Public/User/SocialUser.h
--- a/Engine/Plugins/Online/OnlineFramework/Source/Party/Public/User/SocialUser.h
+++ b/Engine/Plugins/Online/OnlineFramework/Source/Party/Public/User/SocialUser.h
@@ -44,6 +44,7 @@ class USocialUser : public UObject
 
 public:
 	UE_API USocialUser();
+	UE_API ~USocialUser();
 
 	UE_API void RegisterInitCompleteHandler(const FOnNewSocialUserInitialized& OnInitializationComplete);
 	bool IsInitialized() const { return bIsInitialized; }
@@ -283,6 +284,8 @@ private:
 
 	TArray<IOnlinePartyJoinInfoConstRef> ReceivedPartyInvites;
 
+	FDelegateHandle OnReadUserInfoV2CVarChanged;
+
 	// Initialization delegates that fire only when a specific user has finishing initializing
 	TArray<FOnNewSocialUserInitialized> UserInitializedEvents;
 
