Batch Copy Support - General for MVP Launch

Updated text on WBP_VoiceChat_ControlsMoved_CoachMark
Use sentence case for the timespan text format in WBP_TextChat_Restrictions and WBP_SocialBan_Message
Removed second line from DM welcome message
Updated Pending invites text
Updated SetPartyChannelInfo to show member count instead of names and added SetGameChannelInfo to set the status in the Game channel button
Updated interactions to use sentance case

#doc https://docs.google.com/document/d/189j102aAb4cgr3nbxo7bctaX9z-XWrky8YY-ymUmHmI/edit?tab=t.zdnqfhlpl8eb

#tests Open sidebar,verified that the text is correct
#changelist validated
#virtualized
#Synced-CL 48810155

#Synced-CL 48825208

diff --git a/Engine/Plugins/Online/OnlineFramework/Source/Party/Private/Interactions/CoreInteractions.cpp b/Engine/Plugins/Online/OnlineFramework/Source/Party/Private/Interactions/CoreInteractions.cpp
--- a/Engine/Plugins/Online/OnlineFramework/Source/Party/Private/Interactions/CoreInteractions.cpp
+++ b/Engine/Plugins/Online/OnlineFramework/Source/Party/Private/Interactions/CoreInteractions.cpp
@@ -19,9 +19,9 @@ FText FSocialInteraction_AddFriend::GetDisplayName(const USocialUser& User)
 {
 	if (User.IsFriend(ESocialSubsystem::Platform))
 	{
-		return LOCTEXT("AddEpicFriend", "Add Epic Friend");
+		return LOCTEXT("AddEpicFriend", "Add Epic friend");
 	}
-	return LOCTEXT("AddFriend", "Add Friend");
+	return LOCTEXT("AddFriend", "Add friend");
 }
 
 FString FSocialInteraction_AddFriend::GetSlashCommandToken()
@@ -82,7 +82,7 @@ DEFINE_SOCIAL_INTERACTION(RemoveFriend)
 
 FText FSocialInteraction_RemoveFriend::GetDisplayName(const USocialUser& User)
 {
-	return LOCTEXT("RemoveFriend", "Remove Friend");
+	return LOCTEXT("RemoveFriend", "Remove friend");
 }
 
 FString FSocialInteraction_RemoveFriend::GetSlashCommandToken()
diff --git a/Engine/Plugins/Online/OnlineFramework/Source/Party/Private/Interactions/PartyInteractions.cpp b/Engine/Plugins/Online/OnlineFramework/Source/Party/Private/Interactions/PartyInteractions.cpp
--- a/Engine/Plugins/Online/OnlineFramework/Source/Party/Private/Interactions/PartyInteractions.cpp
+++ b/Engine/Plugins/Online/OnlineFramework/Source/Party/Private/Interactions/PartyInteractions.cpp
@@ -17,7 +17,7 @@ DEFINE_SOCIAL_INTERACTION(InviteToParty)
 
 FText FSocialInteraction_InviteToParty::GetDisplayName(const USocialUser& User)
 {
-	return LOCTEXT("InviteToParty", "Invite to Party");
+	return LOCTEXT("InviteToParty", "Invite to party");
 }
 
 FString FSocialInteraction_InviteToParty::GetSlashCommandToken()
