Add UpdateManager timing regions
#Synced-CL 0

diff --git a/Engine/Plugins/Online/OnlineFramework/Source/Hotfix/Private/UpdateManager.cpp b/Engine/Plugins/Online/OnlineFramework/Source/Hotfix/Private/UpdateManager.cpp
--- a/Engine/Plugins/Online/OnlineFramework/Source/Hotfix/Private/UpdateManager.cpp
+++ b/Engine/Plugins/Online/OnlineFramework/Source/Hotfix/Private/UpdateManager.cpp
@@ -13,6 +13,7 @@
 #include "Engine/World.h"
 
 #include "PatchCheck.h"
+#include "ProfilingDebugging/CsvProfiler.h"
 #include "ProfilingDebugging/LoadTimeTracker.h"
 
 #include UE_INLINE_GENERATED_CPP_BY_NAME(UpdateManager)
@@ -606,7 +607,22 @@ void UUpdateManager::SetUpdateState(EUpdateState NewState)
 {
 	if (CurrentUpdateState != NewState)
 	{
-		UE_LOG(LogHotfixManager, Display, TEXT("Update State %s -> %s"), UpdateStateEnum ? *UpdateStateEnum->GetNameStringByValue((int64)CurrentUpdateState) : TEXT("Invalid"), UpdateStateEnum ? *UpdateStateEnum->GetNameStringByValue((int64)NewState) : TEXT("Invalid"));
+		FString OldStateStr = UpdateStateEnum ? *UpdateStateEnum->GetNameStringByValue((int64)CurrentUpdateState) : TEXT("Invalid");
+		FString NewStateStr = UpdateStateEnum ? *UpdateStateEnum->GetNameStringByValue((int64)NewState) : TEXT("Invalid");
+
+#if CSV_PROFILER_STATS || MISCTRACE_ENABLED
+		// Add timing regions
+		if (CurrentUpdateState != EUpdateState::UpdateIdle)
+		{
+			CSV_TRACE_REGION_END(*OldStateStr, TEXT("UpdateManager"))
+		}
+		if (NewState != EUpdateState::UpdateIdle && NewState != EUpdateState::UpdateComplete)
+		{
+			CSV_TRACE_REGION_BEGIN(*NewStateStr, TEXT("UpdateManager"))
+		}
+#endif
+
+		UE_LOG(LogHotfixManager, Display, TEXT("Update State %s -> %s"), *OldStateStr, *NewStateStr);
 		CurrentUpdateState = NewState;
 		OnUpdateStatusChanged().Broadcast(NewState);
 	}
