Fix for the list of recent players never gettfing queried again after start

Issue:
When joining a party in the lobby that contains at least one non-friend, the list of recent players isn't updated until the client is restarted

Casused By:
The SocialToolKit only queries the list of recent players when the client launches and when the sub game changes.

Fix:
When joining a party that contains at least one non-friend, or when a non-friend joins into the player's party, the SocialToolKit updates the recent player list (unless all non-friends are already in that list, in which case we can skip the update)
Added a CVar 'Social.QueryRecentPlayersAfterPartyUpdates' to turn off this new behavior in case it causes any issues.

#tests Tested with various combinations
- Joining a party that already contains at least one non-friend
- Having a non-friend join an existing party
- Verified that the recent player list is only updated for the non-friends (not for the mutual friend) and that it is updated only once.

#rnx
#Synced-CL 48874555

diff --git a/Engine/Plugins/Online/OnlineFramework/Source/Party/Private/SocialToolkit.cpp b/Engine/Plugins/Online/OnlineFramework/Source/Party/Private/SocialToolkit.cpp
--- a/Engine/Plugins/Online/OnlineFramework/Source/Party/Private/SocialToolkit.cpp
+++ b/Engine/Plugins/Online/OnlineFramework/Source/Party/Private/SocialToolkit.cpp
@@ -6,6 +6,8 @@
 #include "SocialQuery.h"
 #include "User/SocialUserList.h"
 #include "Chat/SocialChatManager.h"
+#include "Party/SocialParty.h"
+#include "Party/PartyMember.h"
 
 #include "OnlineSubsystemUtils.h"
 #include "Engine/LocalPlayer.h"
@@ -30,6 +32,13 @@ namespace FriendAcceptFailureReason
 	const FString AcceptFailReason_NotPendingInbound = TEXT("NotPendingInbound");
 }
 
+bool bQueryRecentPlayersAfterPartyUpdates = true;
+static FAutoConsoleVariableRef CVarQueryRecentPlayersAfterPartyUpdates(
+	TEXT("Social.QueryRecentPlayersAfterPartyUpdates"), 
+	bQueryRecentPlayersAfterPartyUpdates,
+	TEXT("Query list of recent players when a player joins a party containing at least one non-friend, or a non-friend joins the player's existing party (unless all of the non-friended players are already in the list of recent players)"),
+	ECVF_Default);
+
 #if WITH_EDITOR
 DECLARE_MULTICAST_DELEGATE_TwoParams(FOnStartRandomizeUserPresence, uint8 /*NumRandomUser*/, float /*TickerTimer*/);
 static FOnStartRandomizeUserPresence Debug_OnStartRandomizeUserPresenceEvent;
@@ -702,6 +711,7 @@ void USocialToolkit::OnOwnerLoggedIn()
 				PartyInterface->AddOnPartyInviteRemovedExDelegate_Handle(FOnPartyInviteRemovedExDelegate::CreateUObject(this, &USocialToolkit::HandlePartyInviteRemoved));
 				PartyInterface->AddOnPartyRequestToJoinReceivedDelegate_Handle(FOnPartyRequestToJoinReceivedDelegate::CreateUObject(this, &USocialToolkit::HandlePartyRequestToJoinReceived));
 				PartyInterface->AddOnPartyRequestToJoinRemovedDelegate_Handle(FOnPartyRequestToJoinRemovedDelegate::CreateUObject(this, &USocialToolkit::HandlePartyRequestToJoinRemoved));
+				PartyInterface->AddOnPartyDataReceivedDelegate_Handle(FOnPartyDataReceivedDelegate::CreateUObject(this, &USocialToolkit::HandlePartyDataReceived));
 			}
 
 			if (IOnlinePresencePtr PresenceInterface = OSS->GetPresenceInterface())
@@ -782,6 +792,7 @@ void USocialToolkit::OnOwnerLoggedOut()
 				PartyInterface->ClearOnPartyInviteRemovedExDelegates(this);
 				PartyInterface->ClearOnPartyRequestToJoinReceivedDelegates(this);
 				PartyInterface->ClearOnPartyRequestToJoinRemovedDelegates(this);
+				PartyInterface->ClearOnPartyDataReceivedDelegates(this);
 			}
 
 			IOnlinePresencePtr PresenceInterface = OSS->GetPresenceInterface();
@@ -1354,6 +1365,40 @@ void USocialToolkit::HandlePartyInviteRemoved(const FUniqueNetId& LocalUserId, c
 	}
 }
 
+void USocialToolkit::HandlePartyDataReceived(const FUniqueNetId& LocalUserId, const FOnlinePartyId& PartyId, const FName& Namespace, const FOnlinePartyData& PartyData)
+{
+	if (bQueryRecentPlayersAfterPartyUpdates)
+	{
+		if (LocalUserId == GetLocalUserNetId(ESocialSubsystem::Primary))
+		{
+			bool RecentPlayerListRequiresUpdate = false;
+			if (const USocialParty* SocialParty = GetSocialManager().GetParty(PartyId))
+			{
+				for (const UPartyMember* Member : SocialParty->GetPartyMembers())
+				{
+					const USocialUser& SocialUser = Member->GetSocialUser();
+					const FUniqueNetIdRepl SocialUserId = SocialUser.GetUserId(ESocialSubsystem::Primary);
+					if (!SocialUserId.IsValid() || *SocialUserId == LocalUserId)
+					{
+						continue;
+					}
+					if (!SocialUser.IsFriend() && !SocialUser.IsRecentPlayer())
+					{
+						RecentPlayerListRequiresUpdate = true;
+						break;
+					}
+				}
+			}
+
+			if (RecentPlayerListRequiresUpdate)
+			{
+				// Get the latest from the IOnlineFriends system
+				QueryRecentPlayers();
+			}
+		}
+	}
+}
+
 void USocialToolkit::HandleBlockPlayerComplete(int32 LocalUserNum, bool bWasSuccessful, const FUniqueNetId& BlockedPlayerId, const FString& ListName, const FString& ErrorStr, ESocialSubsystem SubsystemType)
 {
 	if (bWasSuccessful && LocalUserNum == GetLocalUserNum())
diff --git a/Engine/Plugins/Online/OnlineFramework/Source/Party/Public/SocialToolkit.h b/Engine/Plugins/Online/OnlineFramework/Source/Party/Public/SocialToolkit.h
--- a/Engine/Plugins/Online/OnlineFramework/Source/Party/Public/SocialToolkit.h
+++ b/Engine/Plugins/Online/OnlineFramework/Source/Party/Public/SocialToolkit.h
@@ -285,6 +285,8 @@ private:	// Handlers
 
 	UE_API void HandlePartyInviteReceived(const FUniqueNetId& LocalUserId, const IOnlinePartyJoinInfo& Invite);
 	UE_API void HandlePartyInviteRemoved(const FUniqueNetId& LocalUserId, const IOnlinePartyJoinInfo& Invite, EPartyInvitationRemovedReason Reason);
+	UE_API void HandlePartyDataReceived(const FUniqueNetId& LocalUserId, const FOnlinePartyId& PartyId, const FName& Namespace, const FOnlinePartyData& PartyData);
+
 #if PARTY_PLATFORM_INVITE_PERMISSIONS
 	UE_API void CanReceiveInviteFrom(USocialUser& SocialUser, const IOnlinePartyJoinInfoConstRef& Invite, TFunction<void(const bool /*bResult*/)>&& CompletionFunc);
 #endif
