Batch Copy Support - General for 39.30

Removed trailing '.' from reporting status system message
Added fallback cases for MembersJoined and MembersLeft system messages
Updated interactions to use sentance case

#doc https://docs.google.com/document/d/1tBVgHmFSyfrwl44h2NBmX0twe5BLzNpsHjgmABY7F2I/edit?usp=sharing

#tests Open sidebar, verified that the correct text is displayed
#changelist validated
#virtualized
#Synced-CL 49383107

diff --git a/Engine/Plugins/Online/OnlineFramework/Source/Party/Private/Interactions/PartyInteractions.cpp b/Engine/Plugins/Online/OnlineFramework/Source/Party/Private/Interactions/PartyInteractions.cpp
--- a/Engine/Plugins/Online/OnlineFramework/Source/Party/Private/Interactions/PartyInteractions.cpp
+++ b/Engine/Plugins/Online/OnlineFramework/Source/Party/Private/Interactions/PartyInteractions.cpp
@@ -43,7 +43,7 @@ DEFINE_SOCIAL_INTERACTION(JoinParty)
 
 FText FSocialInteraction_JoinParty::GetDisplayName(const USocialUser& User)
 {
-	return LOCTEXT("JoinParty", "Join Party");
+	return LOCTEXT("JoinParty", "Join party");
 }
 
 FString FSocialInteraction_JoinParty::GetSlashCommandToken()
@@ -70,7 +70,7 @@ DEFINE_SOCIAL_INTERACTION(RequestToJoinParty)
 
 FText FSocialInteraction_RequestToJoinParty::GetDisplayName(const USocialUser& User)
 {
-	return LOCTEXT("RequestToJoinParty", "Request to Join");
+	return LOCTEXT("RequestToJoinParty", "Request to join");
 }
 
 FString FSocialInteraction_RequestToJoinParty::GetSlashCommandToken()
@@ -96,7 +96,7 @@ DEFINE_SOCIAL_INTERACTION(AcceptJoinRequest)
 
 FText FSocialInteraction_AcceptJoinRequest::GetDisplayName(const USocialUser& User)
 {
-	return LOCTEXT("AcceptJoinRequest", "Accept Request");
+	return LOCTEXT("AcceptJoinRequest", "Accept request");
 }
 
 FString FSocialInteraction_AcceptJoinRequest::GetSlashCommandToken()
@@ -122,7 +122,7 @@ DEFINE_SOCIAL_INTERACTION(DismissJoinRequest)
 
 FText FSocialInteraction_DismissJoinRequest::GetDisplayName(const USocialUser& User)
 {
-	return LOCTEXT("DismissJoinRequest", "Dismiss Request");
+	return LOCTEXT("DismissJoinRequest", "Dismiss request");
 }
 
 FString FSocialInteraction_DismissJoinRequest::GetSlashCommandToken()
@@ -148,7 +148,7 @@ DEFINE_SOCIAL_INTERACTION(AcceptPartyInvite)
 
 FText FSocialInteraction_AcceptPartyInvite::GetDisplayName(const USocialUser& User)
 {
-	return LOCTEXT("AcceptPartyInvite", "Accept Invite");
+	return LOCTEXT("AcceptPartyInvite", "Accept invite");
 }
 
 FString FSocialInteraction_AcceptPartyInvite::GetSlashCommandToken()
@@ -175,7 +175,7 @@ DEFINE_SOCIAL_INTERACTION(RejectPartyInvite)
 
 FText FSocialInteraction_RejectPartyInvite::GetDisplayName(const USocialUser& User)
 {
-	return LOCTEXT("RejectPartyInvite", "Reject Invite");
+	return LOCTEXT("RejectPartyInvite", "Reject invite");
 }
 
 FString FSocialInteraction_RejectPartyInvite::GetSlashCommandToken()
@@ -201,7 +201,7 @@ DEFINE_SOCIAL_INTERACTION(LeaveParty)
 
 FText FSocialInteraction_LeaveParty::GetDisplayName(const USocialUser& User)
 {
-	return LOCTEXT("LeaveParty", "Leave Party");
+	return LOCTEXT("LeaveParty", "Leave party");
 }
 
 FString FSocialInteraction_LeaveParty::GetSlashCommandToken()
