Statically cache some OnlineFramework cvars that are read frequently

#rnx
#Synced-CL 49641272

diff --git a/Engine/Plugins/Online/OnlineFramework/Source/Party/Private/SocialManager.cpp b/Engine/Plugins/Online/OnlineFramework/Source/Party/Private/SocialManager.cpp
--- a/Engine/Plugins/Online/OnlineFramework/Source/Party/Private/SocialManager.cpp
+++ b/Engine/Plugins/Online/OnlineFramework/Source/Party/Private/SocialManager.cpp
@@ -475,7 +475,8 @@ FJoinPartyResult USocialManager::ValidateJoinTarget(const USocialUser& UserToJoi
 
 bool USocialManager::ShouldReadPresenceFromOnlineServices()
 {
-	if (IConsoleVariable* CVar = IConsoleManager::Get().FindConsoleVariable(TEXT("onlineframework.party.shouldreadpresencefromonlineservices")))
+	static const IConsoleVariable* CVar = IConsoleManager::Get().FindConsoleVariable(TEXT("onlineframework.party.shouldreadpresencefromonlineservices"));
+	if (CVar)
 	{
 		return CVar->GetBool();
 	}
diff --git a/Engine/Plugins/Online/OnlineFramework/Source/Party/Private/User/SocialUser.cpp b/Engine/Plugins/Online/OnlineFramework/Source/Party/Private/User/SocialUser.cpp
--- a/Engine/Plugins/Online/OnlineFramework/Source/Party/Private/User/SocialUser.cpp
+++ b/Engine/Plugins/Online/OnlineFramework/Source/Party/Private/User/SocialUser.cpp
@@ -718,7 +718,8 @@ USocialUser::FSubsystemUserInfo& USocialUser::FindOrCreateSubsystemInfo(const FU
 
 bool USocialUser::ShouldReadUserInfoFromOnlineServices() const
 {
-	if (IConsoleVariable* CVar = IConsoleManager::Get().FindConsoleVariable(TEXT("onlineframework.party.shouldreaduserinfofromonlineservices")))
+	static const IConsoleVariable* CVar = IConsoleManager::Get().FindConsoleVariable(TEXT("onlineframework.party.shouldreaduserinfofromonlineservices"));
+	if (CVar)
 	{
 		return CVar->GetBool();
 	}
