Optimize hotfix check http requests

#Synced-CL 50152864

diff --git a/Engine/Plugins/Online/OnlineFramework/Source/Hotfix/Private/OnlineHotfixManager.cpp b/Engine/Plugins/Online/OnlineFramework/Source/Hotfix/Private/OnlineHotfixManager.cpp
--- a/Engine/Plugins/Online/OnlineFramework/Source/Hotfix/Private/OnlineHotfixManager.cpp
+++ b/Engine/Plugins/Online/OnlineFramework/Source/Hotfix/Private/OnlineHotfixManager.cpp
@@ -226,6 +226,34 @@ UOnlineHotfixManager::~UOnlineHotfixManager()
 	}
 }
 
+UClass* UOnlineHotfixManager::GetConfiguredHotfixManagerClass()
+{
+	UOnlineHotfixManager* DefaultObject = UOnlineHotfixManager::StaticClass()->GetDefaultObject<UOnlineHotfixManager>();
+	FString HotfixManagerClassName = DefaultObject->HotfixManagerClassName;
+	UClass* HotfixManagerClass = LoadClass<UOnlineHotfixManager>(nullptr, *HotfixManagerClassName, nullptr, LOAD_None, nullptr);
+	if (HotfixManagerClass == nullptr)
+	{
+		// Retrieve the static value, before any plugins, hotfixes, or any other dynamic configs have been applied, as a fallback
+		// This is not a standard method for retrieving config values, and should not be used as a general solution in other code
+		FString FallbackClassName;
+		GConfig->FindBranch(TEXT("Engine"), TEXT("DefaultEngine"))->CombinedStaticLayers.GetString(TEXT("/Script/Hotfix.OnlineHotfixManager"), TEXT("HotfixManagerClassName"), FallbackClassName);
+		if (FallbackClassName.Len())
+		{
+			UE_LOG(LogHotfixManager, Display, TEXT("Unable to locate hotfix manager class '%s'. Trying '%s' instead."), *HotfixManagerClassName, *FallbackClassName)
+				HotfixManagerClassName = FallbackClassName;
+			HotfixManagerClass = LoadClass<UOnlineHotfixManager>(nullptr, *HotfixManagerClassName, nullptr, LOAD_None, nullptr);
+		}
+		if (HotfixManagerClass == nullptr)
+		{
+			UE_LOG(LogHotfixManager, Warning, TEXT("Failed to find requested hotfix manager class '%s'. Falling back to default."), *HotfixManagerClassName);
+			// Just use the default class if it couldn't load what was specified
+			HotfixManagerClass = UOnlineHotfixManager::StaticClass();
+		}
+	}
+
+	return HotfixManagerClass;
+}
+
 UOnlineHotfixManager* UOnlineHotfixManager::Get(UWorld* World)
 {
 	UOnlineHotfixManager* DefaultObject = UOnlineHotfixManager::StaticClass()->GetDefaultObject<UOnlineHotfixManager>();
@@ -235,27 +263,7 @@ UOnlineHotfixManager* UOnlineHotfixManager::Get(UWorld* World)
 		UOnlineHotfixManager* HotfixManager = Cast<UOnlineHotfixManager>(OnlineSub->GetNamedInterface(NAME_HotfixManager));
 		if (HotfixManager == nullptr)
 		{
-			FString HotfixManagerClassName = DefaultObject->HotfixManagerClassName;
-			UClass* HotfixManagerClass = LoadClass<UOnlineHotfixManager>(nullptr, *HotfixManagerClassName, nullptr, LOAD_None, nullptr);
-			if (HotfixManagerClass == nullptr)
-			{
-				// Retrieve the static value, before any plugins, hotfixes, or any other dynamic configs have been applied, as a fallback
-				// This is not a standard method for retrieving config values, and should not be used as a general solution in other code
-				FString FallbackClassName;
-				GConfig->FindBranch(TEXT("Engine"), TEXT("DefaultEngine"))->CombinedStaticLayers.GetString(TEXT("/Script/Hotfix.OnlineHotfixManager"), TEXT("HotfixManagerClassName"), FallbackClassName);
-				if (FallbackClassName.Len())
-				{
-					UE_LOG(LogHotfixManager, Display, TEXT("Unable to locate hotfix manager class '%s'. Trying '%s' instead."), *HotfixManagerClassName, *FallbackClassName)
-					HotfixManagerClassName = FallbackClassName;
-					HotfixManagerClass = LoadClass<UOnlineHotfixManager>(nullptr, *HotfixManagerClassName, nullptr, LOAD_None, nullptr);
-				}
-				if (HotfixManagerClass == nullptr)
-				{
-					UE_LOG(LogHotfixManager, Warning, TEXT("Failed to find requested hotfix manager class '%s'. Falling back to default."), *HotfixManagerClassName);
-					// Just use the default class if it couldn't load what was specified
-					HotfixManagerClass = UOnlineHotfixManager::StaticClass();
-				}
-			}
+			UClass* HotfixManagerClass = GetConfiguredHotfixManagerClass();
 			// Create it and store it
 			HotfixManager = NewObject<UOnlineHotfixManager>(GetTransientPackage(), HotfixManagerClass);
 			OnlineSub->SetNamedInterface(NAME_HotfixManager, HotfixManager);
diff --git a/Engine/Plugins/Online/OnlineFramework/Source/Hotfix/Public/OnlineHotfixManager.h b/Engine/Plugins/Online/OnlineFramework/Source/Hotfix/Public/OnlineHotfixManager.h
--- a/Engine/Plugins/Online/OnlineFramework/Source/Hotfix/Public/OnlineHotfixManager.h
+++ b/Engine/Plugins/Online/OnlineFramework/Source/Hotfix/Public/OnlineHotfixManager.h
@@ -270,15 +270,6 @@ protected:
 	 */
 	UE_API bool IsCompatibleHotfixFile(const FString& InFilename, FString& OutFilename);
 
-	/**
-	 * Override this method to look at the file information for any game specific hotfix processing
-	 * NOTE: Make sure to call Super to get default handling of files
-	 *
-	 * @param FileHeader - the information about the file to determine if it needs custom processing
-	 *
-	 * @return true if the file needs some kind of processing, false to have hotfixing ignore the file
-	 */
-	UE_API virtual bool WantsHotfixProcessing(const FCloudFileHeader& FileHeader);
 	/**
 	 * Called when a file needs custom processing (see above). Override this to provide your own processing methods
 	 *
@@ -407,6 +398,16 @@ public:
 	UE_API UOnlineHotfixManager(FVTableHelper& Helper);
 	UE_API virtual ~UOnlineHotfixManager();
 
+	/**
+	 * Override this method to look at the file information for any game specific hotfix processing
+	 * NOTE: Make sure to call Super to get default handling of files
+	 *
+	 * @param FileHeader - the information about the file to determine if it needs custom processing
+	 *
+	 * @return true if the file needs some kind of processing, false to have hotfixing ignore the file
+	 */
+	UE_API virtual bool WantsHotfixProcessing(const FCloudFileHeader& FileHeader);
+
 	/** Tells the hotfix manager which OSS to use. Uses the default if empty */
 	UPROPERTY(Config)
 	FString OSSName;
@@ -436,6 +437,8 @@ public:
 	 */
 	UE_API virtual void CheckAvailability(FOnHotfixAvailableComplete& InCompletionDelegate);
 
+	static UE_API UClass* GetConfiguredHotfixManagerClass();
+
 	/** Factory method that returns the configured hotfix manager */
 	static UE_API UOnlineHotfixManager* Get(UWorld* World);
 
