Migrate FGenericPlatformMisc::TickHotfixables from actually ticking in LaunchEngineLoop to only being called on hotfix completion

#tests Windows client
#rnx

#Synced-CL 50782910

diff --git a/Engine/Plugins/Online/OnlineFramework/Source/Hotfix/Private/OnlineHotfixManager.cpp b/Engine/Plugins/Online/OnlineFramework/Source/Hotfix/Private/OnlineHotfixManager.cpp
--- a/Engine/Plugins/Online/OnlineFramework/Source/Hotfix/Private/OnlineHotfixManager.cpp
+++ b/Engine/Plugins/Online/OnlineFramework/Source/Hotfix/Private/OnlineHotfixManager.cpp
@@ -907,6 +907,7 @@ void UOnlineHotfixManager::TriggerHotfixComplete(EHotfixResult HotfixResult)
 
 	FThreadHeartBeat::Get().InitSettings();
 	FGameThreadHitchHeartBeat::Get().InitSettings();
+	FGenericPlatformMisc::TickHotfixables();
 
 	TriggerOnHotfixCompleteDelegates(HotfixResult);
 	if (HotfixResult == EHotfixResult::Failed)
diff --git a/Engine/Source/Runtime/Launch/Private/LaunchEngineLoop.cpp b/Engine/Source/Runtime/Launch/Private/LaunchEngineLoop.cpp
--- a/Engine/Source/Runtime/Launch/Private/LaunchEngineLoop.cpp
+++ b/Engine/Source/Runtime/Launch/Private/LaunchEngineLoop.cpp
@@ -5581,11 +5581,6 @@ void FEngineLoop::Tick()
 		FGameThreadHitchHeartBeat::Get().FrameStart();
 	}
 
-	{
-		SCOPED_NAMED_EVENT(Tick_Hotfixables, SystemColor);
-		FPlatformMisc::TickHotfixables();
-	}
-
 	// Must be called on the game thread outside of any frame breadcrumbs.
 	LatchRenderThreadConfiguration();
 
