Change to presence debugging accessibility to non-shipping instead of Editor only
#tests tested in a preflight on non-shipping build and in PIE/Cooked builds

#Synced-CL 50788406

diff --git a/Engine/Plugins/Online/OnlineFramework/Source/Party/Private/SocialToolkit.cpp b/Engine/Plugins/Online/OnlineFramework/Source/Party/Private/SocialToolkit.cpp
--- a/Engine/Plugins/Online/OnlineFramework/Source/Party/Private/SocialToolkit.cpp
+++ b/Engine/Plugins/Online/OnlineFramework/Source/Party/Private/SocialToolkit.cpp
@@ -39,7 +39,7 @@ static FAutoConsoleVariableRef CVarQueryRecentPlayersAfterPartyUpdates(
 	TEXT("Query list of recent players when a player joins a party containing at least one non-friend, or a non-friend joins the player's existing party (unless all of the non-friended players are already in the list of recent players)"),
 	ECVF_Default);
 
-#if WITH_EDITOR
+#if !UE_BUILD_SHIPPING
 DECLARE_MULTICAST_DELEGATE_TwoParams(FOnStartRandomizeUserPresence, uint8 /*NumRandomUser*/, float /*TickerTimer*/);
 static FOnStartRandomizeUserPresence Debug_OnStartRandomizeUserPresenceEvent;
 static FAutoConsoleCommandWithWorldAndArgs CMD_OnStartRandomUserPresence
@@ -199,7 +199,7 @@ USocialToolkit::USocialToolkit()
 
 void USocialToolkit::InitializeToolkit(ULocalPlayer& InOwningLocalPlayer)
 {
-#if WITH_EDITOR
+#if !UE_BUILD_SHIPPING
 	Debug_OnStartRandomizeUserPresenceEvent.AddUObject(this, &USocialToolkit::Debug_OnStartRandomizeUserPresence);
 	Debug_OnStopRandomizeUserPresenceEvent.AddUObject(this, &USocialToolkit::Debug_OnStopRandomizeUserPresence);
 #endif
@@ -1575,7 +1575,7 @@ void USocialToolkit::HandlePartyRequestToJoinRemoved(const FUniqueNetId& LocalUs
 	}
 }
 
-#if WITH_EDITOR
+#if !UE_BUILD_SHIPPING
 void USocialToolkit::Debug_OnStartRandomizeUserPresence(uint8 NumRandomUser, float TickerTimer)
 {
 	if (Debug_PresenceTickerHandle.IsValid())
@@ -1621,13 +1621,16 @@ bool USocialToolkit::Debug_HandleRandomizeUserPresenceTick(float DeltaTime, uint
 void USocialToolkit::Debug_ChangeRandomUserPresence(uint8 NumRandomUser)
 {
 	const TArray<USocialUser*> SocialUsers = GetAllUsers();
-	for (int32 i = 0; i < NumRandomUser; i++)
+	if (SocialUsers.Num() > 0)
 	{
-		int32 UserIndex = FMath::RandRange(0, SocialUsers.Num() - 1);
-		if (USocialUser* SocialUser = SocialUsers[UserIndex])
+		for (int32 i = 0; i < NumRandomUser; i++)
 		{
-			SocialUser->Debug_RandomizePresence();
-			SocialUser->NotifyPresenceChanged(ESocialSubsystem::Primary);
+			int32 UserIndex = FMath::RandRange(0, SocialUsers.Num() - 1);
+			if (USocialUser* SocialUser = SocialUsers[UserIndex])
+			{
+				SocialUser->Debug_RandomizePresence();
+				SocialUser->NotifyPresenceChanged(ESocialSubsystem::Primary);
+			}
 		}
 	}
 }
diff --git a/Engine/Plugins/Online/OnlineFramework/Source/Party/Private/User/SocialUser.cpp b/Engine/Plugins/Online/OnlineFramework/Source/Party/Private/User/SocialUser.cpp
--- a/Engine/Plugins/Online/OnlineFramework/Source/Party/Private/User/SocialUser.cpp
+++ b/Engine/Plugins/Online/OnlineFramework/Source/Party/Private/User/SocialUser.cpp
@@ -531,7 +531,7 @@ EOnlinePresenceState::Type USocialUser::GetOnlineStatus() const
 		return EOnlinePresenceState::Offline;
 	}
 
-#if WITH_EDITOR
+#if !UE_BUILD_SHIPPING
 	if (bDebug_IsPresenceArtificial && !IsBlocked())
 	{
 		return Debug_RandomPresence;
@@ -575,7 +575,7 @@ UE::Online::EUserPresenceStatus USocialUser::GetOnlineStatusV2() const
 		return Result;
 	}
 
-#if WITH_EDITOR
+#if !UE_BUILD_SHIPPING
 	if (bDebug_IsPresenceArtificial && !IsBlocked())
 	{
 		return Debug_RandomPresenceV2;
@@ -2043,7 +2043,7 @@ void USocialUser::EstablishOssInfo(const TSharedRef<FOnlineRecentPlayer>& InRece
 	}
 }
 
-#if WITH_EDITOR
+#if !UE_BUILD_SHIPPING
 void USocialUser::Debug_RandomizePresence()
 {
 	bDebug_IsPresenceArtificial = true;
diff --git a/Engine/Plugins/Online/OnlineFramework/Source/Party/Private/User/SocialUserList.cpp b/Engine/Plugins/Online/OnlineFramework/Source/Party/Private/User/SocialUserList.cpp
--- a/Engine/Plugins/Online/OnlineFramework/Source/Party/Private/User/SocialUserList.cpp
+++ b/Engine/Plugins/Online/OnlineFramework/Source/Party/Private/User/SocialUserList.cpp
@@ -444,7 +444,7 @@ bool FSocialUserList::EvaluateUserPresence(const USocialUser& User, ESocialSubsy
 				bInSameParty = CurrentParty->ContainsUser(User);
 			}
 
-#if WITH_EDITOR
+#if !UE_BUILD_SHIPPING
 			if (OwnerToolkit->Debug_IsRandomlyChangingPresence())
 			{
 				if (USocialManager::ShouldReadPresenceFromOnlineServices())
diff --git a/Engine/Plugins/Online/OnlineFramework/Source/Party/Public/SocialToolkit.h b/Engine/Plugins/Online/OnlineFramework/Source/Party/Public/SocialToolkit.h
--- a/Engine/Plugins/Online/OnlineFramework/Source/Party/Public/SocialToolkit.h
+++ b/Engine/Plugins/Online/OnlineFramework/Source/Party/Public/SocialToolkit.h
@@ -163,7 +163,7 @@ public:
 	 */
 	virtual bool CanAutoRecreatePersistentParty() const { return IsOwnerLoggedIn(); }
 
-#if WITH_EDITOR
+#if !UE_BUILD_SHIPPING
 	bool Debug_IsRandomlyChangingPresence() const { return bDebug_IsRandomlyChangingUserPresence; }
 #endif
 
@@ -306,7 +306,7 @@ private:	// Handlers
 	UE_API void HandlePartyRequestToJoinReceived(const FUniqueNetId& LocalUserId, const FOnlinePartyId& PartyId, const FUniqueNetId& RequesterId, const IOnlinePartyRequestToJoinInfo& Request);
 	UE_API void HandlePartyRequestToJoinRemoved(const FUniqueNetId& LocalUserId, const FOnlinePartyId& PartyId, const FUniqueNetId& RequesterId, const IOnlinePartyRequestToJoinInfo& Request, EPartyRequestToJoinRemovedReason Reason);
 
-#if WITH_EDITOR
+#if !UE_BUILD_SHIPPING
 	UE_API void Debug_OnStartRandomizeUserPresence(uint8 NumRandomUser, float TickerTimer);
 	UE_API void Debug_OnStopRandomizeUserPresence(bool bClearGeneratedPresence);
 	UE_API bool Debug_HandleRandomizeUserPresenceTick(float DeltaTime, uint8 NumRandomUser);
diff --git a/Engine/Plugins/Online/OnlineFramework/Source/Party/Public/User/SocialUser.h b/Engine/Plugins/Online/OnlineFramework/Source/Party/Public/User/SocialUser.h
--- a/Engine/Plugins/Online/OnlineFramework/Source/Party/Public/User/SocialUser.h
+++ b/Engine/Plugins/Online/OnlineFramework/Source/Party/Public/User/SocialUser.h
@@ -196,7 +196,7 @@ protected:
 	UE_API void NotifyFriendInviteRemoved(ESocialSubsystem SubsystemType);
 	UE_API void NotifyUserUnfriended(ESocialSubsystem SubsystemType);
 
-#if WITH_EDITOR
+#if !UE_BUILD_SHIPPING
 	UE_API void Debug_RandomizePresence();
 	bool bDebug_IsPresenceArtificial = false;
 	EOnlinePresenceState::Type Debug_RandomPresence;
