Fixed platform invitations not being sent along with in-game invites when users are not friends

#Synced-CL 50477682

diff --git a/Engine/Plugins/Online/OnlineFramework/Source/Party/Private/Party/SocialParty.cpp b/Engine/Plugins/Online/OnlineFramework/Source/Party/Private/Party/SocialParty.cpp
--- a/Engine/Plugins/Online/OnlineFramework/Source/Party/Private/Party/SocialParty.cpp
+++ b/Engine/Plugins/Online/OnlineFramework/Source/Party/Private/Party/SocialParty.cpp
@@ -357,8 +357,9 @@ bool USocialParty::TryInviteUser(const USocialUser& UserToInvite, const ESocialP
 		}
 
 		const ULocalPlayer* PerformingPlayer = UserToInvite.GetOwningToolkit().GetOwningLocalPlayerPtr();
-
-		if ((UserPlatformId.IsValid() && bIsOnlineOnPlatform) && 
+		bool  bMustSendPlatformInvite = USocialSettings::MustAlwaysSendPlatformSessionInvites();
+		
+		if ((UserPlatformId.IsValid() && (bMustSendPlatformInvite || bIsOnlineOnPlatform)) && 
 			(!UserPrimaryId.IsValid() || bPreferPlatformInvite) && 
 			!IsInviteRateLimited(UserToInvite, ESocialSubsystem::Platform))
 		{
diff --git a/Engine/Plugins/Online/OnlineFramework/Source/Party/Private/SocialSettings.cpp b/Engine/Plugins/Online/OnlineFramework/Source/Party/Private/SocialSettings.cpp
--- a/Engine/Plugins/Online/OnlineFramework/Source/Party/Private/SocialSettings.cpp
+++ b/Engine/Plugins/Online/OnlineFramework/Source/Party/Private/SocialSettings.cpp
@@ -47,6 +47,12 @@ bool USocialSettings::MustSendPrimaryInvites()
 	return SettingsCDO.bMustSendPrimaryInvites;
 }
 
+bool USocialSettings::MustAlwaysSendPlatformSessionInvites()
+{
+	const USocialSettings& SettingsCDO = *GetDefault<USocialSettings>();
+	return SettingsCDO.bMustAlwaysSendPlatformInvites;
+}
+
 bool USocialSettings::ShouldLeavePartyOnDisconnect()
 {
 	const USocialSettings& SettingsCDO = *GetDefault<USocialSettings>();
diff --git a/Engine/Plugins/Online/OnlineFramework/Source/Party/Public/SocialSettings.h b/Engine/Plugins/Online/OnlineFramework/Source/Party/Public/SocialSettings.h
--- a/Engine/Plugins/Online/OnlineFramework/Source/Party/Public/SocialSettings.h
+++ b/Engine/Plugins/Online/OnlineFramework/Source/Party/Public/SocialSettings.h
@@ -60,6 +60,7 @@ public:
 	static UE_API int32 GetDefaultMaxPartySize();
 	static UE_API bool ShouldPreferPlatformInvites();
 	static UE_API bool MustSendPrimaryInvites();
+	static UE_API bool MustAlwaysSendPlatformSessionInvites();
 	static UE_API bool ShouldLeavePartyOnDisconnect();
 	static UE_API bool ShouldSetDesiredPrivacyOnLocalPlayerBecomesLeader();
 	static UE_API float GetUserListAutoUpdateRate();
@@ -99,6 +100,10 @@ private:
 	/** If true, always send invites using the publisher's system even if already sent via a platform system */
 	UPROPERTY(config)
 	bool bMustSendPrimaryInvites = false;
+	
+	/** Some platforms require a platform invite to be sent along with the in-game invitation, even if the players are not friends in the platform */
+	UPROPERTY(config)
+	bool bMustAlwaysSendPlatformInvites=true;
 
 	/** Should we leave a party when it enters the disconnected state? */
 	UPROPERTY(config)
