Send invitation with metadata for event

#Synced-CL 51046133

diff --git a/Engine/Plugins/Online/OnlineFramework/Source/Party/Public/SocialManager.h b/Engine/Plugins/Online/OnlineFramework/Source/Party/Public/SocialManager.h
--- a/Engine/Plugins/Online/OnlineFramework/Source/Party/Public/SocialManager.h
+++ b/Engine/Plugins/Online/OnlineFramework/Source/Party/Public/SocialManager.h
@@ -229,6 +229,7 @@ protected:
 
 	UE_API USocialParty* GetPartyInternal(const FOnlinePartyTypeId& PartyTypeId, bool bIncludeLeavingParties = false) const;
 
+	UE_API virtual void HandleWorldEstablished(UWorld* World);
 private:
 	UE_API UGameInstance& GetGameInstance() const;
 	UE_API USocialToolkit& CreateSocialToolkit(ULocalPlayer& OwningLocalPlayer, int32 LocalPlayerIndex);
@@ -245,7 +246,6 @@ private:
 private:	// Handlers
 	UE_API void HandleGameViewportInitialized();
 	UE_API void HandlePreClientTravel(const FString& PendingURL, ETravelType TravelType, bool bIsSeamlessTravel);
-	UE_API void HandleWorldEstablished(UWorld* World);
 	UE_API void HandleLocalPlayerAdded(int32 LocalUserNum);
 	UE_API void HandleLocalPlayerRemoved(int32 LocalUserNum);
 	UE_API void HandleToolkitReset(int32 LocalUserNum);
