Perf reporting: rename the "UpdateManager" timing region category to "Hotfix", add a Prewarm region
#tests tested on Android

#Synced-CL 51003583

diff --git a/Engine/Plugins/Online/OnlineFramework/Source/Hotfix/Private/UpdateManager.cpp b/Engine/Plugins/Online/OnlineFramework/Source/Hotfix/Private/UpdateManager.cpp
--- a/Engine/Plugins/Online/OnlineFramework/Source/Hotfix/Private/UpdateManager.cpp
+++ b/Engine/Plugins/Online/OnlineFramework/Source/Hotfix/Private/UpdateManager.cpp
@@ -614,11 +614,11 @@ void UUpdateManager::SetUpdateState(EUpdateState NewState)
 		// Add timing regions
 		if (CurrentUpdateState != EUpdateState::UpdateIdle)
 		{
-			CSV_TRACE_REGION_END(*OldStateStr, TEXT("UpdateManager"))
+			CSV_TRACE_REGION_END(*OldStateStr, TEXT("Hotfix"))
 		}
 		if (NewState != EUpdateState::UpdateIdle && NewState != EUpdateState::UpdateComplete)
 		{
-			CSV_TRACE_REGION_BEGIN(*NewStateStr, TEXT("UpdateManager"))
+			CSV_TRACE_REGION_BEGIN(*NewStateStr, TEXT("Hotfix"))
 		}
 #endif
 
