Fix GameFeaturePlugin config not hotfixing after initial startup

Adds cvar hotfix.AlwaysReapply to non-shipping builds for testing.

#rnx
#Synced-CL 51288979

diff --git a/Engine/Plugins/Online/OnlineFramework/Source/Hotfix/Private/OnlineHotfixManager.cpp b/Engine/Plugins/Online/OnlineFramework/Source/Hotfix/Private/OnlineHotfixManager.cpp
--- a/Engine/Plugins/Online/OnlineFramework/Source/Hotfix/Private/OnlineHotfixManager.cpp
+++ b/Engine/Plugins/Online/OnlineFramework/Source/Hotfix/Private/OnlineHotfixManager.cpp
@@ -63,6 +63,14 @@ static TAutoConsoleVariable<bool> CVarEmptyDynamicHotfixContentsOnNewHotfix(
 	TEXT("If true clear out the existing values for DynamicHotfixContents when a new hotfix is processed. When set to false only remove the files which were removed since the last hotfix."),
 	ECVF_Default);
 
+#if !UE_BUILD_SHIPPING
+/* static */ bool UOnlineHotfixManager::bAlwaysReapplyHotfixes = false;
+
+static FAutoConsoleVariableRef CVarAlwaysReapplyHotfixes(
+	TEXT("hotfix.AlwaysReapply"),
+	UOnlineHotfixManager::bAlwaysReapplyHotfixes,
+	TEXT("Whether hotfixes should be downloaded and reapplied on every check, for debugging"));
+#endif // !UE_BUILD_SHIPPING
 
 class FPakFileVisitor : public IPlatformFile::FDirectoryVisitor
 {
@@ -674,7 +682,7 @@ void UOnlineHotfixManager::BuildHotfixFileListDeltas()
 	for (const FCloudFileHeader& CurrentHeader : HotfixFileList)
 	{
 		bool bFoundMatch = LastHotfixFileList.Contains(CurrentHeader);
-		if (!bFoundMatch)
+		if (!bFoundMatch || bAlwaysReapplyHotfixes)
 		{
 			// All NEW or CHANGED ini files will be added to the process list
 			ChangedHotfixFileList.Add(CurrentHeader);
@@ -1053,6 +1061,14 @@ FConfigBranch* UOnlineHotfixManager::GetBranch(const FString& IniName)
 	// find the branch by basename or full filename
 	FConfigBranch* Branch = GConfig->FindBranch(*StrippedIniNameNoExtension, StrippedIniNameNoExtension);
 	
+	if (!Branch)
+	{
+		// Find branch for game feature plugin config hotfixes.
+#define FIND_BRANCH(BranchName) if (!Branch && StrippedIniNameNoExtension.EndsWith(FString(#BranchName))) { Branch = GConfig->FindBranch(FName(#BranchName), FString()); }
+		ENUMERATE_KNOWN_INI_FILES(FIND_BRANCH)
+#undef FIND_BRANCH
+	}
+
 	return Branch;
 }
 
diff --git a/Engine/Plugins/Online/OnlineFramework/Source/Hotfix/Private/UpdateManager.cpp b/Engine/Plugins/Online/OnlineFramework/Source/Hotfix/Private/UpdateManager.cpp
--- a/Engine/Plugins/Online/OnlineFramework/Source/Hotfix/Private/UpdateManager.cpp
+++ b/Engine/Plugins/Online/OnlineFramework/Source/Hotfix/Private/UpdateManager.cpp
@@ -175,7 +175,7 @@ UUpdateManager::EUpdateStartResult UUpdateManager::StartUpdateCheckInternal(cons
 
 		const double CacheTimer = CVarUpdateManagerCacheTimeoutSeconds.GetValueOnGameThread();
 		const double TimeSinceCheck = DeltaTime.GetTotalSeconds();
-		if (bForceCheck || TimeSinceCheck >= CacheTimer)
+		if (bForceCheck || TimeSinceCheck >= CacheTimer || UOnlineHotfixManager::bAlwaysReapplyHotfixes)
 		{
 			if (TSharedPtr<IUpdateManagerStatsCollector> StrongStatsCollector = StatsCollector.Pin())
 			{
diff --git a/Engine/Plugins/Online/OnlineFramework/Source/Hotfix/Public/OnlineHotfixManager.h b/Engine/Plugins/Online/OnlineFramework/Source/Hotfix/Public/OnlineHotfixManager.h
--- a/Engine/Plugins/Online/OnlineFramework/Source/Hotfix/Public/OnlineHotfixManager.h
+++ b/Engine/Plugins/Online/OnlineFramework/Source/Hotfix/Public/OnlineHotfixManager.h
@@ -93,6 +93,15 @@ class UOnlineHotfixManager :
 {
 	GENERATED_BODY()
 
+public:
+
+	// Whether to reapply all hotfixes every time they're rechecked (for debugging).
+#if UE_BUILD_SHIPPING
+	static constexpr bool bAlwaysReapplyHotfixes = false;
+#else
+	static bool bAlwaysReapplyHotfixes;
+#endif
+
 protected:
 	/** The online interface to use for downloading the hotfix files */
 	IOnlineTitleFilePtr OnlineTitleFile;
