[40.10] Add new delegate FPartyInviteReceivedEvent to USocialToolkit

This new delegate FPartyInviteReceivedEvent broadcasts IOnlinePartyJoinInfo, which we now need to distinguish a 1v1 Party invite from a normal party invite.

#testing Loaded into PIE frontend with two clients, P1 and P2, with fro.bEnableRivals1v1 set to true. As P1, attempted to send 1v1 party invite to P2. On P2 client, I now see the IOnlinePartyJoinInfo being passed through to the party invite delegate callback.
#rnx

#Synced-CL 51664842

#Synced-CL 51706803

diff --git a/Engine/Plugins/Online/OnlineFramework/Source/Party/Private/SocialToolkit.cpp b/Engine/Plugins/Online/OnlineFramework/Source/Party/Private/SocialToolkit.cpp
--- a/Engine/Plugins/Online/OnlineFramework/Source/Party/Private/SocialToolkit.cpp
+++ b/Engine/Plugins/Online/OnlineFramework/Source/Party/Private/SocialToolkit.cpp
@@ -468,7 +468,7 @@ void USocialToolkit::NotifyPSNFriendsListRebuilt()
 
 void USocialToolkit::NotifyPartyInviteReceived(USocialUser& SocialUser, const IOnlinePartyJoinInfo& Invite)
 {
-	OnPartyInviteReceived().Broadcast(SocialUser);
+	OnPartyInviteReceived().Broadcast(SocialUser, Invite);
 }
 
 void USocialToolkit::NotifyPartyInviteRemoved(USocialUser& SocialUser, const IOnlinePartyJoinInfo& Invite)
diff --git a/Engine/Plugins/Online/OnlineFramework/Source/Party/Private/User/SocialUserList.cpp b/Engine/Plugins/Online/OnlineFramework/Source/Party/Private/User/SocialUserList.cpp
--- a/Engine/Plugins/Online/OnlineFramework/Source/Party/Private/User/SocialUserList.cpp
+++ b/Engine/Plugins/Online/OnlineFramework/Source/Party/Private/User/SocialUserList.cpp
@@ -142,7 +142,7 @@ void FSocialUserList::HandleOwnerToolkitReset()
 	}
 }
 
-void FSocialUserList::HandlePartyInviteReceived(USocialUser& InvitingUser)
+void FSocialUserList::HandlePartyInviteReceived(USocialUser& InvitingUser, const IOnlinePartyJoinInfo& Invite /* unused */)
 {
 	TryAddUser(InvitingUser);
 	UpdateListInternal();
diff --git a/Engine/Plugins/Online/OnlineFramework/Source/Party/Private/User/SocialUserList.h b/Engine/Plugins/Online/OnlineFramework/Source/Party/Private/User/SocialUserList.h
--- a/Engine/Plugins/Online/OnlineFramework/Source/Party/Private/User/SocialUserList.h
+++ b/Engine/Plugins/Online/OnlineFramework/Source/Party/Private/User/SocialUserList.h
@@ -6,6 +6,7 @@
 #include "Containers/Ticker.h"
 
 class IOnlinePartyRequestToJoinInfo;
+class IOnlinePartyJoinInfo;
 class UPartyMember;
 class USocialParty;
 class USocialToolkit;
@@ -41,7 +42,7 @@ protected:
 private:
 	void HandleOwnerToolkitReset();
 
-	void HandlePartyInviteReceived(USocialUser& InvitingUser);
+	void HandlePartyInviteReceived(USocialUser& InvitingUser, const IOnlinePartyJoinInfo& Invite /* unused */);
 	void HandlePartyInviteRemoved(USocialUser& InvitingUser);
 	void HandlePartyInviteHandled(USocialUser* InvitingUser);
 
diff --git a/Engine/Plugins/Online/OnlineFramework/Source/Party/Public/SocialToolkit.h b/Engine/Plugins/Online/OnlineFramework/Source/Party/Public/SocialToolkit.h
--- a/Engine/Plugins/Online/OnlineFramework/Source/Party/Public/SocialToolkit.h
+++ b/Engine/Plugins/Online/OnlineFramework/Source/Party/Public/SocialToolkit.h
@@ -122,7 +122,8 @@ public:
 	FOnLoginChanged& OnLoginChanged() const { return OnLoginChangedEvent; }
 
 	DECLARE_EVENT_OneParam(USocialToolkit, FPartyInviteEvent, USocialUser&);
-	FPartyInviteEvent& OnPartyInviteReceived() const { return OnPartyInviteReceivedEvent; }
+	DECLARE_EVENT_TwoParams(USocialToolkit, FPartyInviteReceivedEvent, USocialUser&, const IOnlinePartyJoinInfo&);
+	FPartyInviteReceivedEvent& OnPartyInviteReceived() const { return OnPartyInviteReceivedEvent; }
 	FPartyInviteEvent& OnPartyInviteRemoved() const { return OnPartyInviteRemovedEvent; }
 
 	DECLARE_EVENT_TwoParams(USocialToolkit, FFriendInviteEvent, USocialUser&, ESocialSubsystem);
@@ -351,7 +352,7 @@ private:
 
 	mutable FOnLoginChanged OnLoginChangedEvent;
 
-	mutable FPartyInviteEvent OnPartyInviteReceivedEvent;
+	mutable FPartyInviteReceivedEvent OnPartyInviteReceivedEvent;
 	mutable FPartyInviteEvent OnPartyInviteAcceptedEvent;
 	mutable FPartyInviteEvent OnPartyInviteRejectedEvent;
 	mutable FPartyInviteEvent OnPartyInviteRemovedEvent;
