Immediately fail join-in-progress requests to yourself.
Limit join-in-progress responses overall and per target.

diff --git a/Engine/Plugins/Online/OnlineFramework/Source/Party/Private/Party/SocialParty.cpp b/Engine/Plugins/Online/OnlineFramework/Source/Party/Private/Party/SocialParty.cpp
--- a/Engine/Plugins/Online/OnlineFramework/Source/Party/Private/Party/SocialParty.cpp
+++ b/Engine/Plugins/Online/OnlineFramework/Source/Party/Private/Party/SocialParty.cpp
@@ -8,12 +8,16 @@
 #include "SocialManager.h"
 #include "SocialToolkit.h"
 
-#include "PartyBeaconClient.h"
-#include "OnlineSubsystemUtils.h"
+#include "Algo/AnyOf.h"
+#include "Algo/RemoveIf.h"
+#include "Algo/Sort.h"
 #include "Engine/LocalPlayer.h"
 #include "GameFramework/PlayerController.h"
 #include "GameFramework/PlayerState.h"
 #include "Misc/ConfigCacheIni.h"
+#include "OnlineSubsystemUtils.h"
+#include "PartyBeaconClient.h"
+#include "Templates/Greater.h"
 
 #include UE_INLINE_GENERATED_CPP_BY_NAME(SocialParty)
 
@@ -112,6 +116,22 @@ static FAutoConsoleVariableRef CVar_AutoApproveJoinRequests(
 	TEXT("1 Enables. 0 disables."),
 	ECVF_Cheat);
 
+static int32 RejectSelfJoinInProgress = 1;
+static FAutoConsoleVariableRef CVar_RejectSelfJoinInProgress(
+	TEXT("Party.RejectSelfJoinInProgress"),
+	RejectSelfJoinInProgress,
+	TEXT("Reject join-in-progress requests where the target is the requesting local player.\n")
+	TEXT("1 Enables. 0 disables."),
+	ECVF_Default);
+
+static int32 LimitJoinInProgressResponses = 1;
+static FAutoConsoleVariableRef CVar_LimitJoinInProgressResponses(
+	TEXT("Party.LimitJoinInProgressResponses"),
+	LimitJoinInProgressResponses,
+	TEXT("Sort, deduplicate, and cap join-in-progress responses stored in member rep data.\n")
+	TEXT("1 Enables. 0 disables."),
+	ECVF_Default);
+
 bool USocialParty::IsJoiningDuringLoadEnabled()
 {
 	return AllowPartyJoinsDuringLoad != 0;
@@ -1946,6 +1966,19 @@ void USocialParty::JoinSessionCompleteAnalytics(const FSessionId& SessionId, con
 
 void USocialParty::RequestJoinInProgress(const UPartyMember& TargetMember, const FOnRequestJoinInProgressComplete& CompletionDelegate)
 {
+	// Refuse immediately if the target is one of the requesting local players -- joining yourself is not meaningful.
+	const FUniqueNetIdRepl& TargetId = TargetMember.GetPrimaryNetId();
+	const bool bTargetIsSelf = Algo::AnyOf(GetLocalPartyMembersForJoinInProgress(), [&TargetId](const UPartyMember* Member)
+	{
+		return Member && Member->GetPrimaryNetId() == TargetId;
+	});
+	if (RejectSelfJoinInProgress != 0 && bTargetIsSelf)
+	{
+		UE_LOG(LogParty, Warning, TEXT("RequestJoinInProgress: Target %s is the same as the requesting local player -- ignoring."), *TargetId.ToDebugString());
+		CompletionDelegate.ExecuteIfBound(EPartyJoinDenialReason::JoinAttemptAborted);
+		return;
+	}
+
 	// When a join-in-progress attempt is already in progress, cancel the previous attempt before creating a new one.
 	if (IsJoinInProgeressRequestActive())
 	{
@@ -2050,6 +2083,29 @@ void USocialParty::RespondToJoinInProgressRequest(const FPendingMemberApproval&
 		*Response.Requester.ToDebugString(), Response.RequestTime, Response.ResponseTime, ToString(DenialReason));
 	TArray<FPartyMemberJoinInProgressResponse> Responses = GetOwningLocalMember().GetRepData().GetJoinInProgressDataResponses();
 	Responses.Add(Response);
+
+	if (LimitJoinInProgressResponses != 0)
+	{
+		// Sort responses by ResponseTime (descending - newest first).
+		Algo::SortBy(Responses, &FPartyMemberJoinInProgressResponse::ResponseTime, TGreater<int64>());
+
+		// Limit responses per requesting member (keep most recent).
+		{
+			TMap<FUniqueNetIdRepl, int32> ResponseCountPerRequester;
+			Responses.SetNum(Algo::StableRemoveIf(Responses, [&ResponseCountPerRequester, this](const FPartyMemberJoinInProgressResponse& Response)
+			{
+				int32& Count = ResponseCountPerRequester.FindOrAdd(Response.Requester);
+				return ++Count > MaxJoinInProgressResponsesPerMember;
+			}));
+		}
+
+		// Limit total responses (remove oldest, which are at the end after sorting).
+		if (Responses.Num() > MaxJoinInProgressResponses)
+		{
+			Responses.SetNum(MaxJoinInProgressResponses);
+		}
+	}
+
 	GetOwningLocalMember().GetMutableRepData().SetJoinInProgressDataResponses(Responses);
 	RunJoinInProgressTimer();
 }
diff --git a/Engine/Plugins/Online/OnlineFramework/Source/Party/Public/Party/SocialParty.h b/Engine/Plugins/Online/OnlineFramework/Source/Party/Public/Party/SocialParty.h
--- a/Engine/Plugins/Online/OnlineFramework/Source/Party/Public/Party/SocialParty.h
+++ b/Engine/Plugins/Online/OnlineFramework/Source/Party/Public/Party/SocialParty.h
@@ -512,6 +512,14 @@ private:
 	UPROPERTY(config)
 	int32 JoinInProgressResponseTimeout = 60;
 
+	/** Maximum number of join in progress responses to store in member rep data. Oldest responses are removed first. */
+	UPROPERTY(config)
+	int32 MaxJoinInProgressResponses = 12;
+
+	/** Maximum number of join in progress responses to store per requesting member. Oldest responses are removed first. */
+	UPROPERTY(config)
+	int32 MaxJoinInProgressResponsesPerMember = 5;
+
 	mutable FLeavePartyEvent OnPartyLeaveBeginEvent;
 	mutable FLeavePartyEvent OnPartyLeftEvent;
 	mutable FDisconnectPartyEvent OnPartyDisconnectedEvent;
