Fix hotfix manager failing to discover already-active GFPs when processing GFP hotfixes

#rnx
#Synced-CL 51753475

diff --git a/Engine/Plugins/Online/OnlineFramework/Source/Hotfix/Private/OnlineHotfixManager.cpp b/Engine/Plugins/Online/OnlineFramework/Source/Hotfix/Private/OnlineHotfixManager.cpp
--- a/Engine/Plugins/Online/OnlineFramework/Source/Hotfix/Private/OnlineHotfixManager.cpp
+++ b/Engine/Plugins/Online/OnlineFramework/Source/Hotfix/Private/OnlineHotfixManager.cpp
@@ -63,6 +63,13 @@ static TAutoConsoleVariable<bool> CVarEmptyDynamicHotfixContentsOnNewHotfix(
 	TEXT("If true clear out the existing values for DynamicHotfixContents when a new hotfix is processed. When set to false only remove the files which were removed since the last hotfix."),
 	ECVF_Default);
 
+#if !UE_BUILD_SHIPPING
+/* static */ bool UOnlineHotfixManager::bAlwaysReapplyHotfixes = false;
+static FAutoConsoleVariableRef CVarAlwaysReapplyHotfixes(
+	TEXT("hotfix.AlwaysReapply"),
+	UOnlineHotfixManager::bAlwaysReapplyHotfixes,
+	TEXT("Whether hotfixes should be downloaded and reapplied on every check, for debugging"));
+#endif // !UE_BUILD_SHIPPING
 
 class FPakFileVisitor : public IPlatformFile::FDirectoryVisitor
 {
@@ -674,7 +681,7 @@ void UOnlineHotfixManager::BuildHotfixFileListDeltas()
 	for (const FCloudFileHeader& CurrentHeader : HotfixFileList)
 	{
 		bool bFoundMatch = LastHotfixFileList.Contains(CurrentHeader);
-		if (!bFoundMatch)
+		if (!bFoundMatch || bAlwaysReapplyHotfixes)
 		{
 			// All NEW or CHANGED ini files will be added to the process list
 			ChangedHotfixFileList.Add(CurrentHeader);
@@ -1050,9 +1057,26 @@ FConfigBranch* UOnlineHotfixManager::GetBranch(const FString& IniName)
 	const FString StrippedIniName(GetStrippedConfigFileName(IniName));
 	const FString StrippedIniNameNoExtension = FPaths::GetBaseFilename(StrippedIniName);
 
-	// find the branch by basename or full filename
+	// Find the branch by basename or full filename - simple case
 	FConfigBranch* Branch = GConfig->FindBranch(*StrippedIniNameNoExtension, StrippedIniNameNoExtension);
-	
+
+	// If we can't find anything, then we'll look for a branch matching the suffix that has an active plugin tag matching the name.
+	if (!Branch)
+	{
+		const auto FindActivePluginBranch = [&StrippedIniNameNoExtension, &Branch](const FString& BranchName)
+		{
+			if (!Branch && StrippedIniNameNoExtension.EndsWith(BranchName))
+			{
+				const FString ActivePluginName = StrippedIniNameNoExtension.LeftChop(BranchName.Len()); \
+				Branch = GConfig->FindBranchWithDynamicLayerTag(FName(*BranchName), FName(*ActivePluginName));
+			}
+		};
+
+#define FIND_BRANCH(BranchName) FindActivePluginBranch(#BranchName);
+		ENUMERATE_KNOWN_INI_FILES(FIND_BRANCH);
+#undef FIND_BRANCH
+	}
+
 	return Branch;
 }
 
@@ -1108,12 +1132,16 @@ bool UOnlineHotfixManager::HotfixIniFile(const FString& FileName, const FString&
 		FName Tag = *BuildConfigCacheKey(FileName);
 		UE::DynamicConfig::PerformDynamicConfig(Tag, [this, Tag, FileName, IniData](FConfigModificationTracker* ChangeTracker)
 		{
+			// If this file is a known type of ini type, such as DefaultEngine.ini, then we'll find a branch here such as Engine.
+			// This will also work for an active GFP that has applied its config.
 			FConfigBranch* Branch = GetBranch(FileName);
 			if (Branch)
 			{
 				ChangeTracker->CVars.Add(TEXT("ConsoleVariables")).CVarPriority = (int)ECVF_SetByHotfix;
 				Branch->AddDynamicLayerStringToHierarchy(FileName, IniData, Tag, DynamicLayerPriority::Hotfix, ChangeTracker);
 			}
+
+			// If not, then it's probably a GFP hotfix that will be applied whenever the plugin activates.
 			else
 			{
 				const FString HotfixTag = GetStrippedConfigFileName(FileName).Replace(TEXT(".ini"), TEXT(""));
diff --git a/Engine/Plugins/Online/OnlineFramework/Source/Hotfix/Private/UpdateManager.cpp b/Engine/Plugins/Online/OnlineFramework/Source/Hotfix/Private/UpdateManager.cpp
--- a/Engine/Plugins/Online/OnlineFramework/Source/Hotfix/Private/UpdateManager.cpp
+++ b/Engine/Plugins/Online/OnlineFramework/Source/Hotfix/Private/UpdateManager.cpp
@@ -175,7 +175,7 @@ UUpdateManager::EUpdateStartResult UUpdateManager::StartUpdateCheckInternal(cons
 
 		const double CacheTimer = CVarUpdateManagerCacheTimeoutSeconds.GetValueOnGameThread();
 		const double TimeSinceCheck = DeltaTime.GetTotalSeconds();
-		if (bForceCheck || TimeSinceCheck >= CacheTimer)
+		if (bForceCheck || TimeSinceCheck >= CacheTimer || UOnlineHotfixManager::bAlwaysReapplyHotfixes)
 		{
 			if (TSharedPtr<IUpdateManagerStatsCollector> StrongStatsCollector = StatsCollector.Pin())
 			{
diff --git a/Engine/Plugins/Online/OnlineFramework/Source/Hotfix/Public/OnlineHotfixManager.h b/Engine/Plugins/Online/OnlineFramework/Source/Hotfix/Public/OnlineHotfixManager.h
--- a/Engine/Plugins/Online/OnlineFramework/Source/Hotfix/Public/OnlineHotfixManager.h
+++ b/Engine/Plugins/Online/OnlineFramework/Source/Hotfix/Public/OnlineHotfixManager.h
@@ -93,6 +93,14 @@ class UOnlineHotfixManager :
 {
 	GENERATED_BODY()
 
+public:
+	// Whether to reapply all hotfixes every time they're rechecked (for debugging).
+#if UE_BUILD_SHIPPING
+	static constexpr bool bAlwaysReapplyHotfixes = false;
+#else
+	static bool bAlwaysReapplyHotfixes;
+#endif
+
 protected:
 	/** The online interface to use for downloading the hotfix files */
 	IOnlineTitleFilePtr OnlineTitleFile;
@@ -200,7 +208,8 @@ protected:
 	uint32 ChangedOrRemovedPakCount;
 	/** Our passed-in World */
 	TWeakObjectPtr<UWorld> OwnerWorld;
-	/** Loaded hotfix contents that were not mapped to any known branch, but might be loaded later */
+	/** Loaded hotfix contents that were not mapped to any known branch, but might be loaded later.
+	The key used here is a plugin tag + branch name, such as "PluginEngine" for example. */
 	TMap<FName, TArray<TPair<FString, FString>>> DynamicHotfixContents;
 
 	UE_API virtual void Init();
diff --git a/Engine/Source/Runtime/Core/Private/Misc/ConfigCacheIni.cpp b/Engine/Source/Runtime/Core/Private/Misc/ConfigCacheIni.cpp
--- a/Engine/Source/Runtime/Core/Private/Misc/ConfigCacheIni.cpp
+++ b/Engine/Source/Runtime/Core/Private/Misc/ConfigCacheIni.cpp
@@ -4298,6 +4298,24 @@ FConfigBranch* FConfigCacheIni::FindBranchWithNoReload(FName BaseIniName, const
 	return Branch;
 }
 
+FConfigBranch* FConfigCacheIni::FindBranchWithDynamicLayerTag(FName BaseIniName, FName Tag)
+{
+	FConfigBranch* Branch = FindBranch(BaseIniName, "");
+
+	if (Branch)
+	{
+		for (FConfigBranch::DynamicLayerList::TIterator Node(Branch->DynamicLayers.GetHead()); Node; ++Node)
+		{
+			if (Node->Tag == Tag)
+			{
+				return Branch;
+			}
+		}
+	}
+
+	return nullptr;
+}
+
 FConfigBranch* FConfigCacheIni::FindBranch(FName BaseIniName, const FString& Filename)
 {
 	FConfigLogContextScope LogScope(TEXT("FindBranch"), nullptr, nullptr);
diff --git a/Engine/Source/Runtime/Core/Public/Misc/ConfigCacheIni.h b/Engine/Source/Runtime/Core/Public/Misc/ConfigCacheIni.h
--- a/Engine/Source/Runtime/Core/Public/Misc/ConfigCacheIni.h
+++ b/Engine/Source/Runtime/Core/Public/Misc/ConfigCacheIni.h
@@ -1382,6 +1382,9 @@ public:
 	/** Same as FindBranch, but this will not reload a formerly SafeUnload'd branch */
 	CORE_API FConfigBranch* FindBranchWithNoReload(FName BaseIniName, const FString& Filename);
 
+	/** Same as FindBranch, but only if the provided tag is found for one of its dynamic layers */
+	CORE_API FConfigBranch* FindBranchWithDynamicLayerTag(FName BaseIniName, FName Tag);
+
 	/** Create a new branch for FIlename, and return it */
 	CORE_API FConfigBranch& AddNewBranch(const FString& Filename);
 	
