add QoS support for subregion pings map

#rnx
#Synced-CL 51941173

diff --git a/Engine/Plugins/Online/OnlineFramework/Source/Qos/Private/QosInterface.cpp b/Engine/Plugins/Online/OnlineFramework/Source/Qos/Private/QosInterface.cpp
--- a/Engine/Plugins/Online/OnlineFramework/Source/Qos/Private/QosInterface.cpp
+++ b/Engine/Plugins/Online/OnlineFramework/Source/Qos/Private/QosInterface.cpp
@@ -88,6 +88,12 @@ void FQosInterface::GetSubregionPreferences(const FString& RegionId, TArray<FStr
 	return RegionManager->GetSubregionPreferences(RegionId, OutSubregions);
 }
 
+TMap<FString, int64> FQosInterface::GetAllSubregionPingsMap() const
+{
+	check(RegionManager);
+	return RegionManager->GetAllSubregionPingsMap();
+}
+
 void FQosInterface::ForceSelectRegion(const FString& InRegionId)
 {
 	check(RegionManager);
diff --git a/Engine/Plugins/Online/OnlineFramework/Source/Qos/Private/QosRegionManager.cpp b/Engine/Plugins/Online/OnlineFramework/Source/Qos/Private/QosRegionManager.cpp
--- a/Engine/Plugins/Online/OnlineFramework/Source/Qos/Private/QosRegionManager.cpp
+++ b/Engine/Plugins/Online/OnlineFramework/Source/Qos/Private/QosRegionManager.cpp
@@ -681,6 +681,19 @@ void UQosRegionManager::GetSubregionPreferences(const FString& RegionId, TArray<
 	}
 }
 
+TMap<FString, int64> UQosRegionManager::GetAllSubregionPingsMap() const
+{
+	TMap<FString, int64> Pings;
+	for (const FRegionQosInstance& Region : RegionOptions)
+	{
+		for (const FDatacenterQosInstance& Option : Region.DatacenterOptions)
+		{
+			Pings.Add(Option.Definition.Id, Option.AvgPingMs);
+		}
+	}
+	return Pings;
+}
+
 const TArray<FRegionQosInstance>& UQosRegionManager::GetRegionOptions() const
 {
 	if (ForceRegionId.IsEmpty())
@@ -1391,4 +1404,3 @@ FRegionQosInstance UQosRegionManager::TestCreateExampleRegionResult()
 }
 
 #endif // DEBUG_SUBCOMPARE_BY_SUBSPACE
-
diff --git a/Engine/Plugins/Online/OnlineFramework/Source/Qos/Public/QosInterface.h b/Engine/Plugins/Online/OnlineFramework/Source/Qos/Public/QosInterface.h
--- a/Engine/Plugins/Online/OnlineFramework/Source/Qos/Public/QosInterface.h
+++ b/Engine/Plugins/Online/OnlineFramework/Source/Qos/Public/QosInterface.h
@@ -82,6 +82,11 @@ public:
 	 * @param OutSubregions list of subregions in sorted order
 	 */
 	UE_API void GetSubregionPreferences(const FString& RegionId, TArray<FString>& OutSubregions) const;
+	
+	/**
+	 * @return a map of subregion to the computed AvgPing for that subregion
+	 */
+	UE_API TMap<FString, int64> GetAllSubregionPingsMap() const;
 
 	/**
 	 * @return true if this is a usable region, false otherwise
diff --git a/Engine/Plugins/Online/OnlineFramework/Source/Qos/Public/QosRegionManager.h b/Engine/Plugins/Online/OnlineFramework/Source/Qos/Public/QosRegionManager.h
--- a/Engine/Plugins/Online/OnlineFramework/Source/Qos/Public/QosRegionManager.h
+++ b/Engine/Plugins/Online/OnlineFramework/Source/Qos/Public/QosRegionManager.h
@@ -501,6 +501,9 @@ public:
 	 * @param OutSubregions list of subregions in sorted order
 	 */
 	UE_API void GetSubregionPreferences(const FString& RegionId, TArray<FString>& OutSubregions) const;
+	
+	/** @return map of all subregions to avg ping times */
+	UE_API TMap<FString, int64> GetAllSubregionPingsMap() const;
 
 	/**
 	 * @return true if this is a usable region, false otherwise
