Make MarkUserAsDirty protected and expose it

We need to allow social user lists that derive from FSocialUserList to call MarkUserAsDirty on a per-user basis, so that those users can be arbitrarily added/removed from a list in the update flow.

#rnx

#Tests Tested in PIE with local changes that use the newly exposed function in a derived class.
#Synced-CL 51882546

diff --git a/Engine/Plugins/Online/OnlineFramework/Source/Party/Private/User/SocialUserList.h b/Engine/Plugins/Online/OnlineFramework/Source/Party/Private/User/SocialUserList.h
--- a/Engine/Plugins/Online/OnlineFramework/Source/Party/Private/User/SocialUserList.h
+++ b/Engine/Plugins/Online/OnlineFramework/Source/Party/Private/User/SocialUserList.h
@@ -37,6 +37,7 @@ public:
 
 protected:
 	const FSocialUserListConfig& GetListConfig() const { return ListConfig; }
+	virtual void MarkUserAsDirty(USocialUser& User) override;
 
 private:
 	void HandleOwnerToolkitReset();
@@ -63,8 +64,6 @@ private:
 	void HandleRequestToJoinReceived(USocialUser& SocialUser, IOnlinePartyRequestToJoinInfoConstRef Request);
 	void HandleRequestToJoinRemoved(USocialUser& SocialUser, IOnlinePartyRequestToJoinInfoConstRef Request, EPartyRequestToJoinRemovedReason Reason);
 
-	void MarkUserAsDirty(USocialUser& User);
-
 	void TryAddUser(USocialUser& User);
 	void TryAddUserFast(USocialUser& User);
 
diff --git a/Engine/Plugins/Online/OnlineFramework/Source/Party/Public/User/ISocialUserList.h b/Engine/Plugins/Online/OnlineFramework/Source/Party/Public/User/ISocialUserList.h
--- a/Engine/Plugins/Online/OnlineFramework/Source/Party/Public/User/ISocialUserList.h
+++ b/Engine/Plugins/Online/OnlineFramework/Source/Party/Public/User/ISocialUserList.h
@@ -102,4 +102,6 @@ public:
 	virtual void SetAllowAutoUpdate(bool bIsEnabled) = 0;
 
 	virtual void SetAllowSortDuringUpdate(bool bIsEnabled) = 0;
+
+	virtual void MarkUserAsDirty(USocialUser& User) = 0;
 };
